Moving the IASTool command code to it's own module under the 'Developer' directory.

#rnx

- Fixed some compiler/analyzer warnings that started to show up after moving the code as it wasn't being checked before due to the build settings.
-- A minor error with the string format specifier in ::PurlCommandEntry, %d to %u to match the unsigned type being passed into it.
-- Added the same warning supression found in to ::Socks5Peer (proxy server) as can be found in ConnectSocks5 (runtime connection)
-- Minor errors with some string format specifiers in ::Socks5Server, %d to %u to match the unsigned type being passed into it.
- Marked the cache tests as only compiled when IS_PROGRAM is true since the implementation code only compiles under that setting.
- Add UE_API to FOnDemandContainerTocHeaderView::FileContainerFlags so it can be used by the utilities module.
- Move default implementation of FS3Response constructor to the header so that it can be used.
-- Might be better to redo the logic in UploadQueue that requires this but that should be done as it's own work item as it is much riskier.

#Synced-CL 49815329

diff --git a/Engine/Source/Developer/IoStore/OnDemandUtilities/IoStoreOnDemandUtilities.Build.cs b/Engine/Source/Developer/IoStore/OnDemandUtilities/IoStoreOnDemandUtilities.Build.cs
new file mode 100644
--- /dev/null
+++ b/Engine/Source/Developer/IoStore/OnDemandUtilities/IoStoreOnDemandUtilities.Build.cs
@@ -0,0 +1,29 @@
+// Copyright Epic Games, Inc. All Rights Reserved.
+
+using UnrealBuildTool;
+
+public class IoStoreOnDemandUtilities : ModuleRules
+{
+	public IoStoreOnDemandUtilities(ReadOnlyTargetRules Target) : base(Target)
+	{
+		PublicDependencyModuleNames.Add("Core");
+		PublicDependencyModuleNames.Add("TraceLog");
+		PrivateDependencyModuleNames.AddRange(
+			new string[] {
+				"CoreUObject",
+				"IoStoreOnDemand",
+				"IoStoreOnDemandCore",
+				"IoStoreHttpClient",
+				"Json"
+			}
+		);
+
+		CppCompileWarningSettings.UnsafeTypeCastWarningLevel = WarningLevel.Error; 
+
+		if (Target.Platform.IsInGroup(UnrealPlatformGroup.Desktop) &&
+			(Target.Type == TargetType.Editor || Target.Type == TargetType.Program))
+		{
+			PrivateDependencyModuleNames.AddRange(new string[] { "S3Client", "RSA" });
+		}
+	}
+}
diff --git a/Engine/Source/Developer/IoStore/OnDemandUtilities/Private/ChunkPlugin.cpp b/Engine/Source/Developer/IoStore/OnDemandUtilities/Private/ChunkPlugin.cpp
new file mode 100644
--- /dev/null
+++ b/Engine/Source/Developer/IoStore/OnDemandUtilities/Private/ChunkPlugin.cpp
@@ -0,0 +1,973 @@
+// Copyright Epic Games, Inc. All Rights Reserved.
+
+#include "Command.h"
+#include "Common.h"
+
+#include "Algo/Find.h"
+#include "Containers/StringView.h"
+#include "HAL/FileManager.h"
+#include "IO/IoChunkEncoding.h"
+#include "IO/IoContainerHeader.h"
+#include "IO/IoStatus.h"
+#include "IO/IoStore.h"
+#include "IO/IoStoreOnDemand.h"
+#include "Misc/Compression.h"
+#include "Misc/DateTime.h"
+#include "Misc/FileHelper.h"
+#include "Misc/KeyChainUtilities.h"
+#include "Misc/PathViews.h"
+#include "Misc/Paths.h"
+#include "Misc/StringBuilder.h"
+#include "S3/S3Client.h"
+#include "Serialization/JsonReader.h"
+#include "Serialization/JsonSerializer.h"
+#include "Serialization/JsonSerializerMacros.h"
+#include "Serialization/JsonWriter.h"
+#include "Serialization/LargeMemoryWriter.h"
+#include "Serialization/MemoryReader.h"
+#include "UploadQueue.h"
+
+namespace UE::IoStore::Tool
+{
+
+struct FChunkPluginStats : FJsonSerializable
+{
+	double TotalIADChunksSizeKb = 0.0;
+
+BEGIN_JSON_SERIALIZER
+	JSON_SERIALIZE("TotalIADChunksSizeKb", TotalIADChunksSizeKb);
+END_JSON_SERIALIZER
+};
+
+struct FS3Params
+{
+public:
+	FString S3ServiceUrl;
+	FString S3Bucket;
+	FString S3BucketPrefix;
+	FString S3Region;
+	FString S3AccessKey;
+	FString S3SecretKey;
+	FString S3SessionToken;
+	int32 S3MaxConcurrentUploads;
+
+	FS3Params(const FContext& Context)
+	{
+		S3ServiceUrl = FString(Context.Get<FStringView>(TEXT("-ServiceUrl"), FString()));
+		S3Bucket = FString(Context.Get<FStringView>(TEXT("-Bucket"), FString()));
+		S3Region = FString(Context.Get<FStringView>(TEXT("-Region"), FString()));
+		S3AccessKey = FString(Context.Get<FStringView>(TEXT("-AccessKey"), FString()));
+		S3SecretKey = FString(Context.Get<FStringView>(TEXT("-SecretKey"), FString()));
+		S3SessionToken = FString(Context.Get<FStringView>(TEXT("-SessionToken"), FString()));
+
+
+		S3BucketPrefix = FString(Context.Get<FStringView>(TEXT("-BucketPrefix"), FString()));
+		S3MaxConcurrentUploads = Context.Get<int32>(TEXT("-MaxConcurrentUploads"), 10);
+	}
+
+	bool IsValid()
+	{
+		return !S3AccessKey.IsEmpty() && !S3SecretKey.IsEmpty() && !S3SessionToken.IsEmpty();
+	}
+};
+
+
+////////////////////////////////////////////////////////////////////////////////
+struct FChunkPluginSettings
+{
+	TMap<FString, TArray<FString>> PackageSets;
+};
+
+class FChunkWriterInterface
+{
+public:
+	virtual ~FChunkWriterInterface() {}
+	virtual FIoStatus WriteChunk(const FString& RelativeDir, FIoBuffer Chunk, const FIoHash& Hash) = 0;
+
+	virtual bool Flush() = 0;
+};
+
+
+
+class FS3ChunkWriter : public FChunkWriterInterface
+{
+public:
+	FS3ChunkWriter(const FS3Params& S3Params) : 
+		BucketPrefix(S3Params.S3BucketPrefix),
+		Client(FS3ClientConfig({ S3Params.S3Region, S3Params.S3ServiceUrl }), FS3ClientCredentials(S3Params.S3AccessKey, S3Params.S3SecretKey, S3Params.S3SessionToken)),
+		UploadQueue(Client, S3Params.S3Bucket, S3Params.S3MaxConcurrentUploads)
+	{
+	}
+
+	virtual ~FS3ChunkWriter()
+	{
+		// this should have already been done but do it now just in case
+		UploadQueue.Flush();
+	}
+
+
+	virtual FIoStatus WriteChunk(const FString& RelativeDirectory, FIoBuffer Chunk, const FIoHash& Hash) override
+	{
+		TRACE_CPUPROFILER_EVENT_SCOPE(FS3ChunkWriter::WriteChunk);
+
+		const FString HashString = LexToString(Hash);
+
+		TStringBuilder<256> Key;
+		Key << BucketPrefix << TEXT("/")
+			<< RelativeDirectory
+			<< TEXT("/") << HashString.Left(2)
+			<< TEXT("/") << HashString
+			<< TEXT(".iochunk");
+
+		if (UploadQueue.Enqueue(Key, Chunk) == false)
+		{
+			return FIoStatus(EIoErrorCode::WriteError, TEXT("Failed to upload chunk"));
+		}
+
+		return FIoStatus(EIoErrorCode::Ok);
+	}
+
+	virtual bool Flush() override
+	{
+		TRACE_CPUPROFILER_EVENT_SCOPE(FS3ChunkWriter::Flush);
+
+		if (!UploadQueue.Flush())
+		{
+			return false;
+		}
+		return true;
+	}
+
+
+private:
+	FString BucketPrefix;
+	UE::FS3Client Client;
+	FUploadQueue UploadQueue;
+};
+
+
+class FDiskChunkWriter : public FChunkWriterInterface
+{
+public:
+	FDiskChunkWriter(const FString& InOutputFolder)
+		: OutputFolder(InOutputFolder)
+	{
+	}
+
+	virtual ~FDiskChunkWriter() = default;
+
+	virtual  FIoStatus WriteChunk(const FString& Directory, FIoBuffer Chunk, const FIoHash& Hash) override 
+	{
+		TRACE_CPUPROFILER_EVENT_SCOPE(FDiskChunkWriter::WriteChunk);
+
+		IFileManager& FileMgr = IFileManager::Get();
+		const FString HashString = LexToString(Hash);
+
+		TStringBuilder<256> Sb;
+		Sb << OutputFolder << TEXT("/") << Directory << TEXT("/") << HashString.Left(2);
+
+		bool bTree = true;
+		if (FileMgr.MakeDirectory(Sb.ToString(), bTree) == false)
+		{
+			return FIoStatusBuilder(EIoErrorCode::WriteError)
+				<< TEXT("Failed to create directory '")
+				<< FString(Sb.ToString())
+				<< TEXT("'");
+		}
+
+		Sb << TEXT("/") << HashString << TEXT(".iochunk");
+
+		TUniquePtr<FArchive> Ar(IFileManager::Get().CreateFileWriter(Sb.ToString()));
+
+		if (!Ar.IsValid())
+		{
+			return FIoStatusBuilder(EIoErrorCode::WriteError) << TEXT("Failed to open handle to write file '") << FString(Sb.ToString()) << TEXT("'");
+		}
+
+		Ar->Serialize((void*)Chunk.GetView().GetData(), Chunk.GetView().GetSize());
+
+		if (!Ar->Close())
+		{
+			return FIoStatusBuilder(EIoErrorCode::WriteError) << TEXT("Failed to write to file '") << FString(Sb.ToString()) << TEXT("'");
+		}
+
+		UE_LOG(LogIoStoreOnDemand, Log, TEXT("Writing file '%s' (%.2lf KiB)"), Sb.ToString(), double(Chunk.GetView().GetSize()) / 1024.0);
+		return EIoErrorCode::Ok;
+	}
+
+	bool Flush() override
+	{
+		return true;
+	}
+
+private:
+	FString OutputFolder;
+};
+
+////////////////////////////////////////////////////////////////////////////////
+
+static bool ShouldValidateChunk(FName CompressionFormat, uint64 ChunkRawSize)
+{
+	static TMap<FName, bool> ValidatedFormats;
+
+	if (!CompressionFormat.IsNone())
+	{
+		if (bool* bCachedResult = ValidatedFormats.Find(CompressionFormat))
+		{
+			if (*bCachedResult == false)
+			{
+				return false;
+			}
+		}
+		else
+		{
+			// Note that if IsFormatValid returns false it will log a message about it. This is one of the
+			// reasons we must make sure to only check once per format to avoid  log spam. The other reason
+			// is that the check is fairly slow.
+			bool bValid = FCompression::IsFormatValid(CompressionFormat);
+			ValidatedFormats.Add(CompressionFormat, bValid);
+
+			if (bValid == false)
+			{
+				UE_LOG(LogIoStoreOnDemand, Display, TEXT("Skipping validation tests for chunks with the compression format '%s'"), *CompressionFormat.ToString())
+				return false;
+			}
+		}
+	}
+
+	const uint64 MaxChunkSize = 32 * 1024 * 1024;	// We do not want to validate any chunk that will decompress to more than 32 MiB
+													// as we don't want to use too much memory.
+	if (ChunkRawSize >= MaxChunkSize)
+	{
+		return false;
+	}
+
+	return true;
+}
+
+struct FChunkParams
+{
+	FChunkPluginSettings Settings;
+	TMap<FGuid, FAES::FAESKey> EncryptionKeys;
+
+	TUniquePtr<FChunkWriterInterface> ChunkWriter;
+	FString ChunksRelativeFolder;
+
+	bool bIgnoreContainerHeader;
+	bool bValidateChunks;
+};
+
+struct FChunkOutput
+{
+	FChunkOutput(const FString& HostGroupName)
+	{
+		InstallOnDemandToc.Header.Flags = static_cast<uint32>(EOnDemandTocFlags::InstallOnDemand);
+		InstallOnDemandToc.Header.HostGroupName = HostGroupName;
+		InstallOnDemandToc.Header.ChunksDirectory = TEXT(""); // TODO?
+		//InstallOnDemandToc.Containers.Reserve(ContainerFilenames.Num());
+
+		StreamOnDemandToc.Header.Flags = static_cast<uint32>(EOnDemandTocFlags::StreamOnDemand);
+		StreamOnDemandToc.Header.HostGroupName = HostGroupName;
+		StreamOnDemandToc.Header.ChunksDirectory = TEXT(""); // TODO?
+		//StreamOnDemandToc.Containers.Reserve(ContainerFilenames.Num());
+	}
+
+	FOnDemandToc& GetToc(EOnDemandTocFlags Flags)
+	{
+		if (EnumHasAnyFlags(Flags, EOnDemandTocFlags::InstallOnDemand))
+		{
+			return InstallOnDemandToc;
+		}
+		else
+		{
+			return StreamOnDemandToc;
+		}
+	}
+
+	FOnDemandToc InstallOnDemandToc;
+	FOnDemandToc StreamOnDemandToc;
+
+	TMap<FString, TArray<FOnDemandTocTagSetPackageList>> FoundTagSets;
+	TArray<FString> ContainerFiles;
+	FChunkPluginStats Stats;
+};
+
+static int32 ChunkContainer(const FString& ContainerUTocPath, EOnDemandTocFlags TocFlags, const FChunkParams& Input, FChunkOutput& Output)
+{
+	TRACE_CPUPROFILER_EVENT_SCOPE(ChunkPluginCommand::ChunkContainer);
+
+	check(Input.ChunkWriter.IsValid());
+
+	FIoStoreReader ContainerFileReader;
+	{
+		FIoStatus Status = ContainerFileReader.Initialize(*FPaths::ChangeExtension(ContainerUTocPath, TEXT("")), Input.EncryptionKeys);
+		if (!Status.IsOk())
+		{
+			UE_LOG(LogIoStoreOnDemand, Error, TEXT("Failed to open container '%s' for reading due to: %s"), *ContainerUTocPath, *Status.ToString());
+			return -1;
+		}
+	}
+
+	UE_LOG(LogIoStoreOnDemand, Display, TEXT("Processing container '%s'"), *ContainerUTocPath);
+
+	FOnDemandToc& OnDemandToc = Output.GetToc(TocFlags);
+
+	const uint32 BlockSize = ContainerFileReader.GetCompressionBlockSize();
+	if (OnDemandToc.Header.BlockSize == 0)
+	{
+		OnDemandToc.Header.BlockSize = ContainerFileReader.GetCompressionBlockSize();
+	}
+
+	check(OnDemandToc.Header.BlockSize == ContainerFileReader.GetCompressionBlockSize());
+
+	const int32 ContainerIndex = OnDemandToc.Containers.AddDefaulted();
+	FOnDemandTocContainerEntry& ContainerEntry = OnDemandToc.Containers[ContainerIndex];
+	ContainerEntry.ContainerId = ContainerFileReader.GetContainerId();
+	ContainerEntry.ContainerName = FPaths::GetBaseFilename(ContainerUTocPath);
+
+	const bool bIsContainerEncrypted = EnumHasAnyFlags(ContainerFileReader.GetContainerFlags(), EIoContainerFlags::Encrypted);
+	if (bIsContainerEncrypted)
+	{
+		ContainerEntry.EncryptionKeyGuid = LexToString(ContainerFileReader.GetEncryptionKeyGuid());
+		ContainerEntry.ContainerFlags |= static_cast<uint8>(EIoContainerFlags::Encrypted);
+	}
+
+	if (EnumHasAnyFlags(TocFlags, EOnDemandTocFlags::InstallOnDemand) && Input.Settings.PackageSets.Num() > 0)
+	{
+		TArray<FOnDemandTocTagSet> TagSets;
+
+		const FIoChunkId HeaderChunkId = CreateContainerHeaderChunkId(ContainerFileReader.GetContainerId());
+		TIoStatusOr<FIoBuffer> Status = ContainerFileReader.Read(HeaderChunkId, FIoReadOptions());
+		if (!Status.IsOk())
+		{
+			if (Status.Status().GetErrorCode() == EIoErrorCode::NotFound)
+			{
+				UE_LOG(LogIoStoreOnDemand, Display, TEXT("No header chunk for container '%s'"), *ContainerUTocPath);
+			}
+			else
+			{
+				UE_LOG(LogIoStoreOnDemand, Error, TEXT("Failed to read header chunk for container '%s'"), *ContainerUTocPath);
+				return -1;
+			}
+		}
+		else
+		{
+			FIoContainerHeader ContainerHeader;
+			FMemoryReaderView Ar(Status.ValueOrDie().GetView());
+			Ar << ContainerHeader;
+			Ar.Close();
+			if (Ar.IsError() || Ar.IsCriticalError())
+			{
+				UE_LOG(LogIoStoreOnDemand, Error, TEXT("Failed to deserialize header chunk for container '%s'"), *ContainerUTocPath);
+				return -1;
+			}
+
+			for (const TPair<FString, TArray<FString>>& Pair : Input.Settings.PackageSets)
+			{
+				const FString& Tag = Pair.Key;
+
+				FOnDemandTocTagSetPackageList TagSetPackageList;
+				TagSetPackageList.ContainerIndex = ContainerIndex;
+
+				for (const FString& Package : Pair.Value)
+				{
+					const FPackageId PackageId = FPackageId::FromName(FName(Package));
+					int32 PackageIndex = INDEX_NONE;
+					if (ContainerHeader.PackageIds.Find(PackageId, PackageIndex))
+					{
+						UE_LOG(LogIoStoreOnDemand, Log, TEXT("Found package %s for tag %s in container %s"), *Package, *Tag, *ContainerUTocPath);
+						TagSetPackageList.PackageIndicies.Add(PackageIndex);
+					}
+				}
+
+				if (TagSetPackageList.PackageIndicies.Num() > 0)
+				{
+					Output.FoundTagSets.FindOrAdd(Tag).Add(MoveTemp(TagSetPackageList));
+				}
+			}
+		}
+	}
+
+	TArray<FIoStoreTocChunkInfo> ChunkInfos;
+	ContainerFileReader.EnumerateChunks([&ChunkInfos](FIoStoreTocChunkInfo&& Info)
+		{
+			ChunkInfos.Emplace(MoveTemp(Info));
+			return true;
+		});
+
+	UE_LOG(LogIoStoreOnDemand, Display, TEXT("Serializing %d chunks"), ChunkInfos.Num());
+
+	for (const FIoStoreTocChunkInfo& ChunkInfo : ChunkInfos)
+	{
+		const bool bDecrypt = false; // We do not want to decrypt the results. We will upload encrypted chunks and decrypt on the client
+		TIoStatusOr<FIoStoreCompressedReadResult> Status = ContainerFileReader.ReadCompressed(ChunkInfo.Id, FIoReadOptions(), bDecrypt);
+		if (!Status.IsOk())
+		{
+			UE_LOG(LogIoStoreOnDemand, Error, TEXT("Failed to read container chunk, Container='%s', Reason='%s'"),
+				*ContainerUTocPath, *Status.Status().ToString());
+			return -1;
+		}
+
+		FIoStoreCompressedReadResult ReadResult = Status.ConsumeValueOrDie();
+
+		const uint32 BlockOffset = ContainerEntry.BlockSizes.Num();
+		const uint32 BlockCount = ReadResult.Blocks.Num();
+		auto HashBuffer = [](const FIoStoreCompressedReadResult& ReadResult)
+		{
+			TRACE_CPUPROFILER_EVENT_SCOPE(HashBuffer);
+			return FIoHash::HashBuffer(ReadResult.IoBuffer.GetView());
+		};
+		const FIoHash ChunkHash = HashBuffer(ReadResult);
+		const FString HashString = LexToString(ChunkHash);
+
+		FMemoryView EncodedBlocks = ReadResult.IoBuffer.GetView();
+		uint64 RawChunkSize = 0;
+		uint64 EncodedChunkSize = 0;
+		FName CompressionFormat = NAME_None;
+
+		{
+			TRACE_CPUPROFILER_EVENT_SCOPE(HashBlocks);
+			for (const FIoStoreCompressedBlockInfo& BlockInfo : ReadResult.Blocks)
+			{
+				check(Align(BlockInfo.CompressedSize, FAES::AESBlockSize) == BlockInfo.AlignedSize);
+				const uint64 EncodedBlockSize = BlockInfo.AlignedSize;
+				ContainerEntry.BlockSizes.Add(uint32(BlockInfo.CompressedSize));
+
+				FMemoryView EncodedBlock = EncodedBlocks.Left(EncodedBlockSize);
+				EncodedBlocks += EncodedBlock.GetSize();
+
+				check(((char*)EncodedBlock.GetData() + EncodedBlock.GetSize()) <= ((char*)ReadResult.IoBuffer.GetView().GetData() + ReadResult.IoBuffer.GetView().GetSize()));
+				FIoBlockHash Hash = FIoChunkEncoding::HashBlock(EncodedBlock);
+				ContainerEntry.BlockHashes.Add(Hash);
+
+				EncodedChunkSize += EncodedBlockSize;
+				RawChunkSize += BlockInfo.UncompressedSize;
+
+				if (BlockInfo.CompressionMethod != NAME_None)
+				{
+					CompressionFormat = BlockInfo.CompressionMethod;
+					if (OnDemandToc.Header.CompressionFormat.IsEmpty())
+					{
+						OnDemandToc.Header.CompressionFormat = BlockInfo.CompressionMethod.ToString();
+					}
+				}
+			}
+		}
+
+		if (EncodedChunkSize != ReadResult.IoBuffer.GetSize())
+		{
+			UE_LOG(LogIoStoreOnDemand, Error, TEXT("Chunk size mismatch, Container='%s', ChunkId='%s'"),
+				*ContainerUTocPath, *LexToString(ChunkInfo.Id));
+			return -1;
+		}
+
+		// Additional validation to make sure that the chunks can be decrypted and decoded
+		if (Input.bValidateChunks && ShouldValidateChunk(CompressionFormat, RawChunkSize))
+		{
+			TRACE_CPUPROFILER_EVENT_SCOPE(ValidateChunk);
+
+			FIoChunkDecodingParams DecodingParams;
+			DecodingParams.CompressionFormat = CompressionFormat;
+			DecodingParams.TotalRawSize = RawChunkSize;
+			DecodingParams.EncodedBlockSize = TConstArrayView<uint32>(ContainerEntry.BlockSizes.GetData() + BlockOffset, BlockCount);
+			DecodingParams.BlockHash = TConstArrayView<FIoBlockHash>(ContainerEntry.BlockHashes.GetData() + BlockOffset, BlockCount);
+
+			if (bIsContainerEncrypted)
+			{
+				const FAES::FAESKey* EncryptionKey = Input.EncryptionKeys.Find(ContainerFileReader.GetEncryptionKeyGuid());
+				if (EncryptionKey == nullptr)
+				{
+					UE_LOG(LogIoStoreOnDemand, Error, TEXT("Failed to find encryption key '%s', Container='%s"), *LexToString(ContainerFileReader.GetEncryptionKeyGuid()), *ContainerUTocPath);
+					return -1;
+				}
+				DecodingParams.EncryptionKey = FMemoryView(EncryptionKey->Key, EncryptionKey->KeySize);
+			}
+
+			FIoBuffer OutputBuffer(RawChunkSize);
+			bool bDecoded = FIoChunkEncoding::Decode(DecodingParams, ReadResult.IoBuffer.GetView(), OutputBuffer.GetMutableView(), EIoDecodeFlags::None);
+			if (!bDecoded)
+			{
+				UE_LOG(LogIoStoreOnDemand, Error, TEXT("Failed to decode chunk from disk, Container='%s', ChunkId='%s'"), *ContainerUTocPath, *LexToString(ChunkInfo.Id));
+				return -1;
+			}
+		}
+
+		FIoStatus WriteStatus = Input.ChunkWriter->WriteChunk(Input.ChunksRelativeFolder, ReadResult.IoBuffer, ChunkHash);
+		if (WriteStatus.IsOk() == false)
+		{
+			UE_LOG(LogIoStoreOnDemand, Error, TEXT("%s"), *WriteStatus.ToString());
+			return -1;
+		}
+
+		FOnDemandTocEntry& TocEntry = ContainerEntry.Entries.AddDefaulted_GetRef();
+		TocEntry.ChunkId = ChunkInfo.Id;
+		TocEntry.Hash = ChunkHash;
+		TocEntry.RawSize = RawChunkSize;
+		TocEntry.EncodedSize = EncodedChunkSize;
+		TocEntry.BlockOffset = BlockOffset;
+		TocEntry.BlockCount = BlockCount;
+
+		Output.Stats.TotalIADChunksSizeKb += (double)EncodedChunkSize / 1024.0;
+	}
+
+	if (Input.bIgnoreContainerHeader)
+	{
+		UE_LOG(LogIoStoreOnDemand, Display, TEXT("Ignoring container header for '%s'"), *ContainerUTocPath)
+	}
+	else
+	{
+		const bool bDecrypt = false;
+		const FIoChunkId ChunkId = CreateContainerHeaderChunkId(ContainerFileReader.GetContainerId());
+		TIoStatusOr<FIoStoreCompressedReadResult> ReadResult = ContainerFileReader.ReadCompressed(ChunkId, FIoReadOptions(), bDecrypt);
+		if (ReadResult.IsOk())
+		{
+			const FIoBuffer& Chunk = ReadResult.ValueOrDie().IoBuffer;
+			ContainerEntry.Header = MakeArrayView<const uint8>(Chunk.GetData(), static_cast<int32>(Chunk.GetSize()));
+		}
+		else
+		{
+			const EIoErrorCode Err = ReadResult.Status().GetErrorCode();
+			if (Err != EIoErrorCode::UnknownChunkID && Err != EIoErrorCode::NotFound)
+			{
+				UE_LOG(LogIoStoreOnDemand, Error, TEXT("Failed to read container header, reason '%s'"), *ReadResult.Status().ToString())
+					return -1;
+			}
+		}
+	}
+
+	Output.ContainerFiles.Add(ContainerUTocPath);						// .utoc
+	ContainerFileReader.GetContainerFilePaths(Output.ContainerFiles);	// .ucas
+
+	return 0;
+}
+
+int32 WriteOnDemandToc(FOnDemandToc& OnDemandToc, const FString& ContainerName, const FString& OutputFolder)
+{
+	TRACE_CPUPROFILER_EVENT_SCOPE(ChunkPluginCommand::WriteOnDemandToc);
+
+	IFileManager& FileMgr = IFileManager::Get();
+
+	FString Filename = FPathViews::SetExtension(ContainerName, TEXT(".uondemandtoc"));
+	Filename.ToLowerInline();
+
+	const FString TocPath = OutputFolder / Filename;
+
+	if (TUniquePtr<FArchive> Ar(FileMgr.CreateFileWriter(*TocPath)); Ar.IsValid())
+	{
+		*Ar << OnDemandToc;
+		Ar->Close();
+		if (Ar->IsError())
+		{
+			UE_LOG(LogIoStoreOnDemand, Error, TEXT("Failed to serialize TOC '%s'"), *TocPath);
+			return -1;
+		}
+		else
+		{
+			const int64 TocSize = Ar->Tell();
+			UE_LOG(LogIoStoreOnDemand, Display, TEXT("Writing file '%s' (%.2lf KiB)"),
+				*TocPath, double(TocSize) / 1024);
+
+			return 0;
+		}
+	}
+	else
+	{
+		UE_LOG(LogIoStoreOnDemand, Display, TEXT("Failed writing file '%s'"), *TocPath);
+		return -1;
+	}
+}
+
+////////////////////////////////////////////////////////////////////////////////
+static int32 ChunkPluginCommandEntry(const FContext& Context)
+{
+	TRACE_CPUPROFILER_EVENT_SCOPE(ChunkPluginCommand);
+	//while (!FPlatformMisc::IsDebuggerPresent());
+	//UE_DEBUG_BREAK();
+
+	const FString Platform				= FString(Context.Get<FStringView>(TEXT("-Platform"), FString()));
+	const FString BuildVersion			= FString(Context.Get<FStringView>(TEXT("-BuildVersion"), FString()));
+	const FString OnDemandTocName		= FString(Context.Get<FStringView>(TEXT("-OnDemandTocName"), FString()));
+	const FString InputFolder			= FString(Context.Get<FStringView>(TEXT("-InputFolder"), FString()));
+	const FString InOutputFolder		= FString(Context.Get<FStringView>(TEXT("-OutputFolder"), FString()));
+	const FString IntermediateFolder	= FString(Context.Get<FStringView>(TEXT("-IntermediateFolder"), FString()));
+	FString SettingsFile				= FString(Context.Get<FStringView>(TEXT("-SettingsFile"), FString()));
+	FString OutputStatsJson				= FString(Context.Get<FStringView>(TEXT("-OutputStatsJson"), FString()));
+	const FString HostGroupName			= FString(Context.Get<FStringView>(TEXT("-HostGroupName"), FString()));
+
+	const bool bIgnoreContainerHeader	= Context.Get<bool>(TEXT("-IgnoreContainerHeader"), false);
+	const bool bSupportInstalledContainers = !Context.Get<bool>(TEXT("-SkipInstallOnDemand"), false);
+	const bool bSupportStreamingContainers = !Context.Get<bool>(TEXT("-SkipStreamOnDemand"), false);
+
+	FS3Params S3Params(Context);
+
+	const bool bIncludeSigPak			= Context.Get<bool>(TEXT("-IncludeSigPak"), false);
+	const bool bDeleteContainerFiles	= !Context.Get<bool>(TEXT("-KeepContainerFiles"), false);
+	const bool bValidateChunks			= Context.Get<bool>(TEXT("-ValidateChunks"), false);
+
+	FString OutputFolder				= InOutputFolder;
+	FString ContainerFolder				= InputFolder;
+	FString IoStoreRelativeFolder		= "iostore";
+	FString ChunksRelativeFolder		= IoStoreRelativeFolder / TEXT("chunks");
+
+	FPaths::NormalizeDirectoryName(ContainerFolder);
+	FPaths::NormalizeDirectoryName(OutputFolder);
+	FPaths::NormalizeFilename(SettingsFile);
+	FPaths::NormalizeFilename(OutputStatsJson);
+
+	UE_LOG(LogIoStoreOnDemand, Display, TEXT("I/O store chunk plugin:"));
+	UE_LOG(LogIoStoreOnDemand, Display, TEXT("----------------------------------------"));
+	UE_LOG(LogIoStoreOnDemand, Display, TEXT("\tBuildVersion: %s"), *BuildVersion);
+	UE_LOG(LogIoStoreOnDemand, Display, TEXT("\tPlatform: %s"), *Platform);
+	UE_LOG(LogIoStoreOnDemand, Display, TEXT("\tOnDemandTocName: %s"), *OnDemandTocName);
+	UE_LOG(LogIoStoreOnDemand, Display, TEXT("\tInputFolder: %s"), *InputFolder);
+	UE_LOG(LogIoStoreOnDemand, Display, TEXT("\tOutputFolder: %s"), *OutputFolder);
+	UE_LOG(LogIoStoreOnDemand, Display, TEXT("\tIntermediateFolder: %s"), *IntermediateFolder);
+	UE_LOG(LogIoStoreOnDemand, Display, TEXT("\tSkipInstallOnDemand: %s"), !bSupportInstalledContainers ? TEXT("true") : TEXT("false"));
+	UE_LOG(LogIoStoreOnDemand, Display, TEXT("\tSkipStreamOnDemand: %s"), !bSupportStreamingContainers ? TEXT("true") : TEXT("false"));
+	UE_LOG(LogIoStoreOnDemand, Display, TEXT("\tSettingsFile: %s"), *SettingsFile);
+	UE_LOG(LogIoStoreOnDemand, Display, TEXT("\tOutputStatsJson: %s"), *OutputStatsJson);
+	UE_LOG(LogIoStoreOnDemand, Display, TEXT("\tIncludeSigPak: %s"), bIncludeSigPak ? TEXT("true") : TEXT("false"));
+	UE_LOG(LogIoStoreOnDemand, Display, TEXT("\tDeleteContainerFiles: %s"), bDeleteContainerFiles ? TEXT("true") : TEXT("false"));
+	UE_LOG(LogIoStoreOnDemand, Display, TEXT("\tIgnoreContainerHeader: %s"), bIgnoreContainerHeader ? TEXT("true") : TEXT("false"));
+	UE_LOG(LogIoStoreOnDemand, Display, TEXT("\tValidateChunks: %s"), bValidateChunks ? TEXT("true") : TEXT("false"));
+
+	FChunkParams Input;
+	Input.ChunksRelativeFolder = ChunksRelativeFolder;
+	Input.bIgnoreContainerHeader = bIgnoreContainerHeader;
+	Input.bValidateChunks = bValidateChunks;
+
+	IFileManager& FileMgr = IFileManager::Get();
+
+	if (S3Params.IsValid())
+	{
+		Input.ChunkWriter = MakeUnique<FS3ChunkWriter>(S3Params);
+	}
+	else
+	{
+		Input.ChunkWriter = MakeUnique<FDiskChunkWriter>(OutputFolder);
+	}
+
+	if (!SettingsFile.IsEmpty())
+	{
+		UE_LOG(LogIoStoreOnDemand, Display, TEXT("Loading settings file '%s'..."), *SettingsFile);
+
+		TUniquePtr<FArchive> Ar(FileMgr.CreateFileReader(*SettingsFile));
+		if (!Ar || Ar->IsError())
+		{
+			UE_LOG(LogIoStoreOnDemand, Error, TEXT("Failed to open settings file '%s'"), *SettingsFile);
+			return -1;
+		}
+
+		TSharedRef<TJsonReader<UTF8CHAR>> JsonReader = TJsonReaderFactory<UTF8CHAR>::Create(Ar.Get());
+
+		TSharedPtr<FJsonValue> JsonSettings;
+		if (!FJsonSerializer::Deserialize(*JsonReader, JsonSettings))
+		{
+			UE_LOG(LogIoStoreOnDemand, Error, TEXT("Failed to read settings file '%s'"), *SettingsFile);
+			return -1;
+		}
+
+		TSharedPtr<FJsonObject> JsonSettingsObject = JsonSettings->AsObject();
+		if (!JsonSettingsObject)
+		{
+			UE_LOG(LogIoStoreOnDemand, Error, TEXT("Bad settings file '%s'"), *SettingsFile);
+			return -1;
+		}
+
+		if (TSharedPtr<FJsonValue> PackageSets = JsonSettingsObject->TryGetField(TEXT("PackageSets")))
+		{
+			TSharedPtr<FJsonObject> PackageSetsObject = PackageSets->AsObject();
+			if (!PackageSetsObject)
+			{
+				UE_LOG(LogIoStoreOnDemand, Error, TEXT("Bad settings file '%s'"), *SettingsFile);
+				return -1;
+			}
+
+			for (const TPair<FString, TSharedPtr<FJsonValue>>& Pair : PackageSetsObject->Values)
+			{
+				TArray<FString> Packages;
+				if (!PackageSetsObject->TryGetStringArrayField(Pair.Key, Packages))
+				{
+					UE_LOG(LogIoStoreOnDemand, Error, TEXT("Bad settings file '%s'"), *SettingsFile);
+					return -1;
+				}
+
+				UE_LOG(LogIoStoreOnDemand, Display, TEXT("Found Package Set '%s' with %d entries"), *Pair.Key, Packages.Num());
+				for (const FString& Package : Packages)
+				{
+					UE_LOG(LogIoStoreOnDemand, Log, TEXT("->\t'%s'"), *Package);
+				}
+
+				Input.Settings.PackageSets.Add(FString(Pair.Key), MoveTemp(Packages));
+			}
+		}
+	}
+
+	{
+		FKeyChain KeyChain = Common::LoadCryptoKeys(Context);
+		for (const TPair<FGuid, FNamedAESKey>& KeyPair : KeyChain.GetEncryptionKeys())
+		{
+			Input.EncryptionKeys.Add(KeyPair.Key, KeyPair.Value.Key);
+		}
+	}
+
+	if (FileMgr.DirectoryExists(*ContainerFolder) == false)
+	{
+		UE_LOG(LogIoStoreOnDemand, Error, TEXT("Directory '%s' does not exist"), *ContainerFolder);
+		return -1;
+	}
+
+	TArray<FString> ContainerFilenames;
+	FileMgr.FindFiles(ContainerFilenames, *ContainerFolder, TEXT("*.utoc"));
+	UE_LOG(LogIoStoreOnDemand, Display, TEXT("Found %d container files(s)"), ContainerFilenames.Num());
+
+	if (!bSupportInstalledContainers && !bSupportStreamingContainers)
+	{
+		UE_LOG(LogIoStoreOnDemand, Error, TEXT("Neither Installed or Streaming containers were requested"));
+		return -1;
+	}
+
+	FChunkOutput Output(HostGroupName);
+	Output.FoundTagSets.Reserve(Input.Settings.PackageSets.Num());
+
+	for (const FString& Filename : ContainerFilenames)
+	{
+		const FString FullPath = ContainerFolder / Filename;
+
+		// TODO: At the moment the only way we can identify an OnDemand container is by the name. Might be better to pass this info into the tool
+		const bool bIsStreamOnDemandContainer = FullPath.Contains(TEXT("OnDemand"));
+
+		EOnDemandTocFlags ContainerFlags = EOnDemandTocFlags::None;
+
+		if (bSupportStreamingContainers && bIsStreamOnDemandContainer)
+		{
+			ContainerFlags = EOnDemandTocFlags::StreamOnDemand;
+		}
+		else if (bSupportInstalledContainers && !bIsStreamOnDemandContainer)
+		{
+			ContainerFlags = EOnDemandTocFlags::InstallOnDemand;
+		}
+		else
+		{
+			UE_LOGF(LogIoStoreOnDemand, Display, "Skipping container '%ls'", *FullPath);
+			continue;
+		}
+
+		int32 Result = ChunkContainer(FullPath, ContainerFlags, Input, Output);
+		if (Result < 0)
+		{
+			return Result;
+		}
+	}
+
+	check(bSupportInstalledContainers || Output.FoundTagSets.IsEmpty());
+
+	if (bSupportInstalledContainers)
+	{
+		Output.InstallOnDemandToc.TagSets.Reserve(Output.FoundTagSets.Num());
+		for (TPair<FString, TArray<FOnDemandTocTagSetPackageList>>& Pair : Output.FoundTagSets)
+		{
+			FOnDemandTocTagSet& TocTagSet = Output.InstallOnDemandToc.TagSets.Emplace_GetRef();
+			TocTagSet.Tag = Pair.Key;
+			TocTagSet.Packages = MoveTemp(Pair.Value);
+		}
+	}
+
+	if (bDeleteContainerFiles)
+	{
+		TRACE_CPUPROFILER_EVENT_SCOPE(ChunkPluginCommand::DeleteContainerFiles);
+
+		UE_LOG(LogIoStoreOnDemand, Display, TEXT("Deleting chunked  containers..."));
+
+		for (const FString& Path : Output.ContainerFiles)
+		{
+			if (FileMgr.FileExists(*Path))
+			{
+				UE_LOG(LogIoStoreOnDemand, Display, TEXT("Deleting '%s'"), *Path);
+				if (!FileMgr.Delete(*Path, /*RequireExists*/true))
+				{
+					UE_LOG(LogIoStoreOnDemand, Error, TEXT("Failed to delete '%s'"), *Path);
+				}
+			}
+		}
+	}
+
+	// Write additional file(s)
+	if (bIncludeSigPak)
+	{
+		TRACE_CPUPROFILER_EVENT_SCOPE(ChunkPluginCommand::IncludeSigPak);
+
+		const FStringView AllowedExt[]
+		{
+			TEXTVIEW("pak"),
+			TEXTVIEW("sig"),
+		};
+
+		TStringBuilder<256> Sb;
+		TArray<FString> AdditionalFiles;
+		FileMgr.FindFiles(AdditionalFiles, *ContainerFolder, TEXT("*.*"));
+
+		UE_LOG(LogIoStoreOnDemand, Display, TEXT("Serializing %d additional file(s)"), AdditionalFiles.Num());
+		for (const FString& Filename : AdditionalFiles)
+		{
+			const bool bIncludeDot = false;
+			const FStringView Ext = FPathViews::GetExtension(Filename, bIncludeDot);
+
+			if (const FStringView* Result =
+				Algo::FindByPredicate(AllowedExt, [&Ext](const FStringView& E) { return Ext == E; });
+				Result == nullptr)
+			{
+				continue;
+			}
+
+			const FString FullPath = ContainerFolder / Filename;
+			TArray<uint8> FileData;
+			if (!FFileHelper::LoadFileToArray(FileData, *FullPath))
+			{
+				UE_LOG(LogIoStoreOnDemand, Error, TEXT("Failed reading file '%s'"), *FullPath);
+				return -1;
+			}
+
+			FMemoryView Chunk = MakeMemoryView(FileData.GetData(), FileData.Num());
+			const FIoHash ChunkHash = FIoHash::HashBuffer(Chunk);
+
+			FIoBuffer Buffer(FIoBuffer::Clone, Chunk);
+
+			FIoStatus WriteStatus = Input.ChunkWriter->WriteChunk(ChunksRelativeFolder, Buffer, ChunkHash);
+
+			if (WriteStatus.IsOk() == false)
+			{
+				UE_LOG(LogIoStoreOnDemand, Error, TEXT("%s"), *WriteStatus.ToString());
+				return -1;
+			}
+
+			UE_LOG(LogIoStoreOnDemand, Display, TEXT("Adding additional file '%s'"), *Filename);
+			Output.InstallOnDemandToc.AdditionalFiles.Add(FOnDemandTocAdditionalFile
+			{
+				.Hash = ChunkHash,
+				.Filename = Filename,
+				.FileSize = IntCastChecked<uint64>(FileData.Num())
+			});
+		}
+	}
+
+	const int64 Timestamp = FDateTime::UtcNow().ToUnixTimestamp();
+	// todo, we actually want the .uondemandtoc files in the base directory
+	// const FString TocPath = OutputFolder / Filename; // like this yo
+	const FString OnDemandTocOutputFolder = OutputFolder / IoStoreRelativeFolder;
+
+	if (bSupportInstalledContainers && !Output.InstallOnDemandToc.Containers.IsEmpty())
+	{
+		Output.InstallOnDemandToc.Meta.EpochTimestamp = Timestamp;
+		Output.InstallOnDemandToc.Meta.TargetPlatform = Platform;
+		Output.InstallOnDemandToc.Meta.BuildVersion = BuildVersion;
+
+		int32 Result = WriteOnDemandToc(Output.InstallOnDemandToc, OnDemandTocName, OnDemandTocOutputFolder);
+		if (Result < 0)
+		{
+			return Result;
+		}
+	}
+
+	if (bSupportStreamingContainers && !Output.StreamOnDemandToc.Containers.IsEmpty())
+	{
+		Output.StreamOnDemandToc.Meta.EpochTimestamp = Timestamp;
+		Output.StreamOnDemandToc.Meta.TargetPlatform = Platform;
+		Output.StreamOnDemandToc.Meta.BuildVersion = BuildVersion;
+
+		// TODO: OnDemandTocName
+		int32 Result = WriteOnDemandToc(Output.StreamOnDemandToc, OnDemandTocName + TEXT("ondemand"), OnDemandTocOutputFolder);
+		if (Result < 0)
+		{
+			return Result;
+		}
+	}
+
+	// Write dummy containers if necessary
+	if (bDeleteContainerFiles)
+	{
+		TRACE_CPUPROFILER_EVENT_SCOPE(ChunkPluginCommand::WriteDummyContainerFiles);
+
+		UE_LOG(LogIoStoreOnDemand, Display, TEXT("Replacing .utoc for chunked  containers..."));
+
+		for (const FString& SourcePath : Output.ContainerFiles)
+		{
+			if (!SourcePath.EndsWith(TEXT(".utoc")))
+			{
+				continue;
+			}
+
+			FNameBuilder NameBuilder;
+			NameBuilder << BuildVersion << FPathViews::GetBaseFilename(SourcePath) << TEXTVIEW("dummy");
+
+			FIoContainerSettings ContainerSettings;
+			ContainerSettings.ContainerId = FIoContainerId::FromName(FName(NameBuilder));
+
+			const FString FullPath = ContainerFolder / FPaths::GetCleanFilename(SourcePath);
+
+			FIoStoreTocResource Toc;
+
+			TIoStatusOr<uint64> Status = FIoStoreTocResource::Write(*FullPath, Toc, 0, 0, ContainerSettings);
+			if (!Status.IsOk())
+			{
+				UE_LOG(LogIoStoreOnDemand, Error, TEXT("Failed to write dummy container '%s' (%s)"), *FullPath, *Status.Status().ToString());
+				return -1;
+			}
+
+			UE_LOG(LogIoStoreOnDemand, Display, TEXT("Wrote dummy file '%s' (%.2lf KiB)"), *FullPath, double(Status.ValueOrDie()) / 1024);
+		}
+	}
+
+
+	if (!Input.ChunkWriter->Flush())
+	{
+		UE_LOG(LogIoStoreOnDemand, Error, TEXT("Writer error: Failed to upload chunk(s)"));
+		return -1; 
+	}
+
+	if (!OutputStatsJson.IsEmpty())
+	{
+		TUniquePtr<FArchive> FileWriter(FileMgr.CreateFileWriter(*OutputStatsJson));
+		if (!FileWriter)
+		{
+			UE_LOG(LogIoStoreOnDemand, Display, TEXT("Failed writing stats file '%s'"), *OutputStatsJson);
+			return -1;
+		}
+		TSharedRef<TJsonWriter<UTF8CHAR>> JsonWriter = TJsonWriterFactory<UTF8CHAR>::Create(FileWriter.Get());
+		Output.Stats.ToJson(JsonWriter, false);
+	}
+
+	return 0;
+}
+
+////////////////////////////////////////////////////////////////////////////////
+extern FArgumentSet S3Arguments;
+
+static FCommand ChunkPluginCommand(
+	ChunkPluginCommandEntry,
+	TEXT("ChunkPlugin"),
+	TEXT(""),
+	{
+		TArgument<FStringView>(TEXT("-Platform"),			TEXT("Platform name.")),
+		TArgument<FStringView>(TEXT("-BuildVersion"),		TEXT("Build version")),
+		TArgument<FStringView>(TEXT("-OnDemandTocName"),	TEXT("On Demand TOC Name")),
+		TArgument<FStringView>(TEXT("-InputFolder"),		TEXT("Input folder to plugin information.")),
+		TArgument<FStringView>(TEXT("-OutputFolder"),		TEXT("Ouptut folder.")),
+		TArgument<FStringView>(TEXT("-IntermediateFolder"),	TEXT("Intermediate folder.")),
+		TArgument<FStringView>(TEXT("-SettingsFile"),		TEXT("Optional settings file.")),
+		TArgument<FStringView>(TEXT("-OutputStatsJson"),	TEXT("Path to write a json file with statistics.")),
+		TArgument<FStringView>(TEXT("-HostGroupName"),		TEXT("Host group name or URL")),
+		TArgument<FStringView>(TEXT("-CryptoKeys"),			TEXT("")),
+		TArgument<bool>(TEXT("-IncludeSigPak"),				TEXT("Include .sig and .pak file in the uondemandtoc")),
+		TArgument<bool>(TEXT("-KeepContainerFiles"),		TEXT("Should we keep the container files after processing them.")),
+		TArgument<bool>(TEXT("-SkipStreamOnDemand"),		TEXT("Do not chunk any stream on-demand containers")),
+		TArgument<bool>(TEXT("-SkipInstallOnDemand"),		TEXT("Do not chunk any install on-demand containers")),
+		TArgument<bool>(TEXT("-IgnoreContainerHeader"),		TEXT("Ignore the container header?")),
+		TArgument<FStringView>(TEXT("-BucketPrefix"),		TEXT("Path to prefix to bucket objects")),
+		TArgument<int32>(TEXT("-MaxConcurrentUploads"),		TEXT("Number of simultaneous uploads")),
+		TArgument<bool>(TEXT("-ValidateChunks"),			TEXT("Runs additional validation on chunks after they are read from the container and before they are written out")),
+		S3Arguments,
+	}
+);
+
+} // namespace UE::IoStore::Tool
diff --git a/Engine/Source/Developer/IoStore/OnDemandUtilities/Private/Command.cpp b/Engine/Source/Developer/IoStore/OnDemandUtilities/Private/Command.cpp
new file mode 100644
--- /dev/null
+++ b/Engine/Source/Developer/IoStore/OnDemandUtilities/Private/Command.cpp
@@ -0,0 +1,579 @@
+// Copyright Epic Games, Inc. All Rights Reserved.
+
+#include "Command.h"
+
+#include <Containers/UnrealString.h>
+#include <CoreGlobals.h>
+#include <Misc/OutputDeviceRedirector.h>
+#include <Misc/ScopeExit.h>
+
+#if PLATFORM_WINDOWS
+#	include "Windows/AllowWindowsPlatformTypes.h"
+#		include <errhandlingapi.h>
+#	include "Windows/HideWindowsPlatformTypes.h"
+#endif
+
+namespace UE::IoStore::Tool {
+
+////////////////////////////////////////////////////////////////////////////////
+namespace Error
+{
+
+struct FBase
+{
+					FBase(const TCHAR* Ever, FStringView Not) : What(Ever), Why(Not) {}
+	const TCHAR*	What;
+	FString			Why;
+};
+
+#define MAKE_IASTOOL_ERROR(Name) \
+	struct Name	: public FBase { \
+		Name(FStringView Not) : FBase(TEXT(#Name), Not) {} \
+	};
+	MAKE_IASTOOL_ERROR(FMissingValue);
+	MAKE_IASTOOL_ERROR(FCommandNotFound);
+	MAKE_IASTOOL_ERROR(FBoolPositional);
+	MAKE_IASTOOL_ERROR(FUnexpectedValue);
+	MAKE_IASTOOL_ERROR(FUnknownKey);
+	MAKE_IASTOOL_ERROR(FCommandAbort);
+#undef MAKE_IASTOOL_ERROR
+
+}
+
+////////////////////////////////////////////////////////////////////////////////
+template <typename ErrorType, typename... ArgTypes>
+requires std::is_base_of<Error::FBase, ErrorType>::value
+[[noreturn]] void FatalError(ArgTypes&&... Args)
+{
+	throw ErrorType(Forward<ArgTypes>(Args)...);
+}
+
+
+
+////////////////////////////////////////////////////////////////////////////////
+static void WriteLine(const TCHAR* Format, ...)
+{
+	va_list VaList;
+	va_start(VaList, Format);
+	ON_SCOPE_EXIT { va_end(VaList); };
+
+	TCHAR Buffer[256];
+	FCString::GetVarArgs(Buffer, UE_ARRAY_COUNT(Buffer), Format, VaList);
+
+	static FName Category = FName("IasTool");
+	GLog->Serialize(Buffer, ELogVerbosity::Type::Display, Category);
+}
+
+////////////////////////////////////////////////////////////////////////////////
+static bool IsAnAstleySingleFrom91(const TCHAR* Cassette)
+{
+	if (Cassette[0] != '-' && Cassette[0] != '/')
+	{
+		return false;
+	}
+
+	FStringView Cd = Cassette + 1;
+	for (const TCHAR* TwelveInch : { TEXT("help"), TEXT("h"), TEXT("?"), TEXT("-help") }) // '/-help' :)
+	{
+		if (Cd.Equals(TwelveInch, ESearchCase::IgnoreCase))
+		{
+			return true;
+		}
+	}
+
+	return false;
+}
+
+
+
+////////////////////////////////////////////////////////////////////////////////
+using ParseFunc = void (void*, FStringView);
+
+////////////////////////////////////////////////////////////////////////////////
+template <typename Type>
+struct TArgumentParse
+{
+	static ParseFunc* GetParser()
+	{
+		return [] (void* Out, FStringView Input) {
+			auto* Inner = (Type*)Out;
+			LexFromString(*Inner, Input);
+		};
+	}
+};
+
+////////////////////////////////////////////////////////////////////////////////
+template <>
+struct TArgumentParse<bool>
+{
+	static ParseFunc* GetParser() { return nullptr; }
+};
+
+////////////////////////////////////////////////////////////////////////////////
+template <>
+struct TArgumentParse<FStringView>
+{
+	static ParseFunc* GetParser()
+	{
+		return [] (void* Out, FStringView Input) {
+			auto* Inner = (FStringView*)Out;
+			new (Inner) FStringView(Input);
+		};
+	}
+};
+
+////////////////////////////////////////////////////////////////////////////////
+static void Parse(const FArgument& Argument, void* Out, FStringView Input)
+{
+	static_assert(sizeof(PTRINT) == sizeof(ParseFunc*));
+	auto* Parser = (ParseFunc*)(Argument.Inner);
+	return Parser(Out, Input);
+}
+
+////////////////////////////////////////////////////////////////////////////////
+template <typename Type>
+static FArgument MakeArgument(FStringView Name, FStringView Desc)
+{
+	static_assert(alignof(Type) <= FArgument::ValueAlign);
+	static_assert(sizeof(Type) <= FArgument::ValueSize);
+	return { Name, Desc, PTRINT(TArgumentParse<Type>::GetParser()) };
+}
+
+#define TARGUMENT_IMPL(t) \
+	template <> FArgument TArgument<t>(FStringView Name, FStringView Desc) { \
+		return MakeArgument<t>(Name, Desc); \
+	}
+	TARGUMENT_IMPL(FStringView)
+	TARGUMENT_IMPL(bool)
+	TARGUMENT_IMPL(uint8)	TARGUMENT_IMPL(int8)
+	TARGUMENT_IMPL(uint16)	TARGUMENT_IMPL(int16)
+	TARGUMENT_IMPL(uint32)	TARGUMENT_IMPL(int32)
+	TARGUMENT_IMPL(uint64)	TARGUMENT_IMPL(int64)
+	TARGUMENT_IMPL(float)	TARGUMENT_IMPL(double)
+#undef TARGUMENT_IMPL
+
+
+
+////////////////////////////////////////////////////////////////////////////////
+FArgumentSet::FArgumentSet(std::initializer_list<FArgument> InArguments)
+: Arguments(InArguments)
+{
+}
+
+////////////////////////////////////////////////////////////////////////////////
+FArgumentSet::operator FArgument () const
+{
+	return { FStringView(), FStringView(), 0 - PTRINT(this) };
+}
+
+
+
+////////////////////////////////////////////////////////////////////////////////
+FContext::FContext(const FArguments& InArguments, int32 ArgC, const TCHAR* const* ArgV)
+{
+	SetArguments(InArguments);
+
+	auto NextArg = [&ArgC, &ArgV, i=int32(1), Pending=FStringView()] () mutable -> FStringView
+	{
+		ArgC--;
+		if (auto Ret = Pending; !Ret.IsEmpty())
+		{
+			Pending.Reset();
+			return Ret.Mid(Ret[0] == '=');
+		}
+		const TCHAR* Arg = ArgV[i++];
+		if (Arg[0] != '-')
+		{
+			return FStringView(Arg);
+		}
+		const TCHAR* End = Arg;
+		for (uint32 c = *End; c == '-' || uint32('z' - (c | 0x20)) <= 26; c = *++End);
+		Pending = FStringView(End);
+		ArgC += (Pending.IsEmpty() == false);
+		return FStringView(Arg, int32(ptrdiff_t(End - Arg)));
+	};
+
+	int32 PendingOptional = -1;
+	auto HandleOptional = [this, &PendingOptional] (FStringView Arg)
+	{
+		if (int32 i = PendingOptional; i >= 0)
+		{
+			if (Arg.IsEmpty() || Arg[0] == '-')
+			{
+				FArgument Argument = Arguments[i];
+				FatalError<Error::FMissingValue>(Argument.Name);
+			}
+			Parse(Arguments[i], Values[i].Buffer, Arg);
+			SetMask |= 1ull << uint64(i);
+			PendingOptional = -1;
+			return;
+		}
+
+		for (int32 i = 0, n = Arguments.Num(); i < n; ++i)
+		{
+			FArgument& Argument = Arguments[i];
+			if (Argument.Name != Arg)
+			{
+				continue;
+			}
+
+			if (Argument.Inner != 0)
+			{
+				PendingOptional = i;
+				return;
+			}
+
+			Values[i].Buffer[0] = 1;
+			SetMask |= 1ull << uint64(i);
+			return;
+		}
+
+		FatalError<Error::FUnknownKey>(Arg);
+	};
+
+	int32 PositionalIndex = -1;
+	auto HandlePositional = [this, &PositionalIndex] (FStringView Arg)
+	{
+		for (;;)
+		{
+			++PositionalIndex;
+			if (PositionalIndex >= Arguments.Num())
+			{
+				FatalError<Error::FUnexpectedValue>(Arg);
+			}
+
+			if (Arguments[PositionalIndex].Name[0] != '-')
+			{
+				break;
+			}
+		}
+
+		FArgument& Argument = Arguments[PositionalIndex];
+		if (Argument.Inner == 0)
+		{
+			FatalError<Error::FBoolPositional>(Arg);
+		}
+
+		Parse(Argument, Values[PositionalIndex].Buffer, Arg);
+		SetMask |= 1ull << uint64(PositionalIndex);
+	};
+
+	Values.SetNum(Arguments.Num());
+
+	while (ArgC > 1)
+	{
+		FStringView Arg = NextArg();
+		if (Arg == TEXT("--"))
+		{
+			break;
+		}
+
+		(PendingOptional >= 0 || Arg[0] == '-')
+			? HandleOptional(Arg)
+			: HandlePositional(Arg);
+	}
+
+	if (PendingOptional >= 0)
+	{
+		FArgument Argument = Arguments[PendingOptional];
+		FatalError<Error::FMissingValue>(Argument.Name);
+	}
+}
+
+////////////////////////////////////////////////////////////////////////////////
+void FContext::SetArguments(const FArguments& Input)
+{
+	for (const FArgument& Item : Input)
+	{
+		if (Item.Inner >= 0)
+		{
+			Arguments.Add(Item);
+			continue;
+		}
+
+		const auto* Set = (FArgumentSet*)(0 - Item.Inner);
+		SetArguments(Set->Arguments);
+	}
+}
+
+////////////////////////////////////////////////////////////////////////////////
+uint32 FContext::FindArgument(FStringView Name) const
+{
+	for (const FArgument& Argument : Arguments)
+	{
+		if (Argument.Name == Name)
+		{
+			return uint32(ptrdiff_t(&Argument - Arguments.GetData()));
+		}
+	}
+
+	FatalError<Error::FUnknownKey>(Name);
+}
+
+////////////////////////////////////////////////////////////////////////////////
+[[noreturn]] void FContext::OnMissingValue(FStringView Name) const
+{
+	FatalError<Error::FMissingValue>(Name);
+}
+
+////////////////////////////////////////////////////////////////////////////////
+[[noreturn]] void FContext::Abort(const TCHAR* Reason) const
+{
+	FatalError<Error::FCommandAbort>(Reason);
+}
+
+////////////////////////////////////////////////////////////////////////////////
+template <typename Lambda>
+static void IterArguments(const FArguments& Arguments, const Lambda& Callback)
+{
+	for (const FArgument& Argument : Arguments)
+	{
+		if (Argument.Inner >= 0)
+		{
+			Callback(Argument);
+			continue;
+		}
+
+		const auto* Set = (FArgumentSet*)(0 - Argument.Inner);
+		IterArguments(Set->Arguments, Callback);
+	}
+}
+
+
+
+////////////////////////////////////////////////////////////////////////////////
+FCommand::FCommand(
+	EntryFunc* InEntry,
+	FStringView InName,
+	FStringView InDesc,
+	FArguments&& InArguments)
+{
+	Initialize(InEntry, InName, InDesc, MoveTemp(InArguments));
+	FCommand::Head() = this;
+}
+
+////////////////////////////////////////////////////////////////////////////////
+void FCommand::Initialize(
+	EntryFunc* InEntry,
+	FStringView InName,
+	FStringView InDesc,
+	FArguments&& InArguments)
+{
+	Entry = InEntry;
+	Name = InName;
+	Desc = InDesc;
+	Arguments = InArguments;
+}
+
+////////////////////////////////////////////////////////////////////////////////
+FCommand*& FCommand::Head()
+{
+	static FCommand* Instance;
+	return Instance;
+}
+
+////////////////////////////////////////////////////////////////////////////////
+FContext FCommand::GetContext(int32 ArgC, const TCHAR* const* ArgV) const
+{
+	return FContext(Arguments, ArgC, ArgV);
+}
+
+////////////////////////////////////////////////////////////////////////////////
+void FCommand::Usage() const
+{
+	if (!Desc.IsEmpty())
+	{
+		WriteLine(TEXT("%s"), *WriteToString<128>(Desc));
+		WriteLine(TEXT(""));
+	}
+
+	bool bHasOpts = false;
+
+	TStringBuilder<128> UsageLine;
+	UsageLine << TEXT("  ");
+	UsageLine << Name;
+	IterArguments(Arguments, [&bHasOpts, &UsageLine] (const FArgument& Argument)
+	{
+		if (Argument.Name.StartsWith('-'))
+		{
+			bHasOpts = true;
+			return;
+		}
+
+		UsageLine << TEXT(" <");
+		UsageLine << Argument.Name;
+		UsageLine << TEXT(">");
+	});
+
+	if (bHasOpts)
+	{
+		UsageLine << TEXT(" [Options]");
+	}
+
+	WriteLine(TEXT("Usage:"));
+	WriteLine(TEXT("  %s"), *UsageLine);
+
+	if (!bHasOpts)
+	{
+		return;
+	}
+
+	WriteLine(TEXT(""));
+	WriteLine(TEXT("Options:"));
+	IterArguments(Arguments, [] (const FArgument& Argument)
+	{
+		if (!Argument.Name.StartsWith('-'))
+		{
+			return;
+		}
+
+		const TCHAR* Suffix = (Argument.Inner != 0) ? TEXT("=<value>") : TEXT("");
+		WriteLine(TEXT( "  %s%s"), Argument.Name.GetData(), Suffix);
+		if (!Argument.Desc.IsEmpty())
+		{
+			WriteLine(TEXT("    %s"), Argument.Desc.GetData());
+			WriteLine(TEXT(""));
+		}
+	});
+}
+
+////////////////////////////////////////////////////////////////////////////////
+int32 FCommand::Call(int32 ArgC, const TCHAR* const* ArgV)
+{
+	for (int32 i = 1; i < ArgC; ++i)
+	{
+		if (IsAnAstleySingleFrom91(ArgV[i]))
+		{
+			Usage();
+			return 126;
+		}
+	}
+
+	FContext Context = GetContext(ArgC, ArgV);
+	return Entry(Context);
+}
+
+////////////////////////////////////////////////////////////////////////////////
+int32 FCommand::Main(int32 ArgC, const TCHAR* const* ArgV)
+{
+#if PLATFORM_WINDOWS
+	::SetUnhandledExceptionFilter([] (EXCEPTION_POINTERS* Pointers) -> LONG
+	{
+		uint32 Code = Pointers->ExceptionRecord->ExceptionCode;
+		void* Where = Pointers->ExceptionRecord->ExceptionAddress;
+		WriteLine(TEXT("ERROR: An exceptional circumstance occurred; %08x @ %p"), Code, Where);
+		check(false);
+		return EXCEPTION_CONTINUE_SEARCH;
+	});
+#endif
+
+	try
+	{
+		return MainInner(ArgC, ArgV);
+	}
+ 	catch (const Error::FBase& Err)
+	{
+		WriteLine(TEXT("ERROR: %s : %s"), Err.What, *(Err.Why));
+	}
+
+	return 125;
+}
+
+////////////////////////////////////////////////////////////////////////////////
+int32 FCommand::MainInner(int32 ArgC, const TCHAR* const* ArgV)
+{
+	if ((ArgC < 2) || IsAnAstleySingleFrom91(ArgV[1]))
+	{
+		WriteLine(TEXT("Available commands:"));
+		for (FCommand* Cmd = Head(); Cmd != nullptr; Cmd = Cmd->Next)
+		{
+			WriteLine(TEXT("  %-12s %s"), Cmd->Name.GetData(), Cmd->Desc.GetData());
+		}
+		return 127;
+	}
+
+	FStringView Action = ArgV[1];
+	for (FCommand* Cmd = Head(); Cmd != nullptr; Cmd = Cmd->Next)
+	{
+		if (Cmd->Name == Action)
+		{
+			return Cmd->Call(ArgC - 1, ArgV + 1);
+		}
+	}
+
+	FatalError<Error::FCommandNotFound>(Action);
+}
+
+
+
+////////////////////////////////////////////////////////////////////////////////
+void CommandTest()
+{
+	auto Entry = [] (const FContext&) { return 0; };
+	FCommand Command;
+	Command.Initialize(Entry, TEXT("name"), TEXT("desc"), {
+		TArgument<FStringView>(TEXT("pos0"), TEXT("description")),
+		TArgument<FStringView>(TEXT("-string"), TEXT("description")),
+		TArgument<FStringView>(TEXT("pos1"), TEXT("description")),
+		TArgument<FStringView>(TEXT("pos2"), TEXT("description")),
+		TArgument<bool>(TEXT("-bool"), TEXT("description")),
+	});
+
+	using TestArgs = TArray<const TCHAR*>;
+
+	auto TestContext = [&] (TestArgs Args)
+	{
+		Args.Insert(TEXT("$"), 0);
+		return Command.GetContext(Args.Num(), Args.GetData());
+	};
+
+	// optional
+	struct {
+		TestArgs	Args;
+		FStringView	ExpectString;
+		bool		ExpectBool;
+	} TestCases0[] = {
+		{ {TEXT("-bool")}, TEXT(""), true },
+		{ {TEXT("-string=abc")}, TEXT("abc"), false },
+		{ {TEXT("-bool"), TEXT("-string=a c")}, TEXT("a c"), true },
+		{ {TEXT("-string"), TEXT("abc")}, TEXT("abc"), false },
+		{ {TEXT("-string"), TEXT("a=c"), TEXT("-bool")}, TEXT("a=c"), true },
+		{ {TEXT("-string=a=c")}, TEXT("a=c"), false },
+	};
+	for (auto& TestCase : TestCases0)
+	{
+		FContext Context = TestContext(MoveTemp(TestCase.Args));
+		check(Context.Get<bool>(TEXT("-bool")) == TestCase.ExpectBool);
+		check(Context.Get<FStringView>(TEXT("-string")) == TestCase.ExpectString);
+	}
+
+	for (TestArgs TestCase : {
+			TestArgs{ TEXT("-string") },
+			TestArgs{ TEXT("-string"), TEXT("-a9e") },
+		})
+	{
+		try { TestContext(MoveTemp(TestCase)); }
+		catch (const Error::FMissingValue&) { continue; }
+		check(false);
+	}
+
+	// positional
+	for (TestArgs TestCase : {
+			TestArgs{ TEXT("pos0") },
+			TestArgs{ TEXT("pos0"), TEXT("pos1") },
+			TestArgs{ TEXT("pos0"), TEXT("pos1"), TEXT("pos2") },
+		})
+	{
+		try { TestContext(MoveTemp(TestCase)); }
+		catch (...) { check(false); }
+	}
+
+	if (bool bTestPass = false; true)
+	{
+		try { TestContext(TestArgs{TEXT("a"), TEXT("b"), TEXT("b"), TEXT("POP!")}); }
+		catch (const Error::FUnexpectedValue&) { bTestPass = true; }
+		check(bTestPass);
+	}
+}
+
+} // namespace UE::IoStore::Tool
diff --git a/Engine/Source/Developer/IoStore/OnDemandUtilities/Private/Command.h b/Engine/Source/Developer/IoStore/OnDemandUtilities/Private/Command.h
new file mode 100644
--- /dev/null
+++ b/Engine/Source/Developer/IoStore/OnDemandUtilities/Private/Command.h
@@ -0,0 +1,135 @@
+// Copyright Epic Games, Inc. All Rights Reserved.
+
+#pragma once
+
+#include <Containers/Array.h>
+#include <Containers/StringConv.h>
+#include <Containers/StringView.h>
+#include <String/LexFromString.h>
+
+#include <initializer_list>
+
+namespace UE::IoStore::Tool {
+
+////////////////////////////////////////////////////////////////////////////////
+struct FArgument
+{
+	static const uint32 ValueAlign = 16;
+	static const uint32 ValueSize = 16;
+	FStringView			Name;
+	FStringView			Desc;
+	PTRINT				Inner;
+};
+
+using FArguments = TArray<FArgument>;
+
+////////////////////////////////////////////////////////////////////////////////
+struct FArgumentSet
+{
+				FArgumentSet(std::initializer_list<FArgument>);
+				operator FArgument () const;
+	FArguments	Arguments;
+};
+
+////////////////////////////////////////////////////////////////////////////////
+template <typename T>
+concept CArgTypes =
+	std::is_same_v<T, bool>
+	|| std::is_integral_v<T>
+	|| std::is_floating_point_v<T>
+	|| std::is_same_v<T, FStringView>;
+
+
+////////////////////////////////////////////////////////////////////////////////
+template <CArgTypes Type> FArgument TArgument(FStringView Name, FStringView Desc);
+
+#define TARGUMENT_DECL(t) \
+	template <> FArgument TArgument<t>(FStringView Name, FStringView Desc);
+
+TARGUMENT_DECL(FStringView)
+TARGUMENT_DECL(bool)
+TARGUMENT_DECL(uint8)	TARGUMENT_DECL(int8)
+TARGUMENT_DECL(uint16)	TARGUMENT_DECL(int16)
+TARGUMENT_DECL(uint32)	TARGUMENT_DECL(int32)
+TARGUMENT_DECL(uint64)	TARGUMENT_DECL(int64)
+TARGUMENT_DECL(float)	TARGUMENT_DECL(double)
+#undef TARGUMENT_DECL
+
+////////////////////////////////////////////////////////////////////////////////
+class FContext
+{
+public:
+	template <CArgTypes Type> Type	Get(FStringView Name) const;
+	template <CArgTypes Type> Type	Get(FStringView Name, Type Default) const;
+	[[noreturn]] void				Abort(const TCHAR* Reason) const;
+
+private:
+	struct alignas(FArgument::ValueAlign) FValue
+	{
+		uint8			Buffer[FArgument::ValueSize] = {};
+	};
+
+	friend				class FCommand;
+						FContext(const FArguments& InArguments, int32 ArgC, const TCHAR* const* ArgV);
+	void				SetArguments(const FArguments& Input);
+	uint32				FindArgument(FStringView Name) const;
+	[[noreturn]] void	OnMissingValue(FStringView Name) const;
+	bool				IsSet(uint32 Index) const { return !!((1ull << Index) & SetMask); }
+	FArguments			Arguments;
+	TArray<FValue>		Values;
+	uint64				SetMask = 0;
+};
+
+////////////////////////////////////////////////////////////////////////////////
+template <CArgTypes Type>
+Type FContext::Get(FStringView Name) const
+{
+	if (Name.StartsWith(TEXT('-')))
+	{
+		return Get(Name, Type{});
+	}
+
+	uint32 Index = FindArgument(Name);
+	if (IsSet(Index))
+	{
+		return *(Type*)(Values[Index].Buffer);
+	}
+
+	OnMissingValue(Name);
+}
+
+////////////////////////////////////////////////////////////////////////////////
+template <CArgTypes Type>
+Type FContext::Get(FStringView Name, Type Default) const
+{
+	uint32 Index = FindArgument(Name);
+	return IsSet(Index) ? *(Type*)(Values[Index].Buffer) : Default;
+}
+
+
+
+////////////////////////////////////////////////////////////////////////////////
+class FCommand
+{
+public:
+	using				EntryFunc = int32 (const FContext&);
+						FCommand(EntryFunc* InEntry, FStringView InName, FStringView InDesc, FArguments&& InArguments);
+	static int32		Main(int32 ArgC, const TCHAR* const* ArgV);
+
+private:
+	friend void			CommandTest();
+						FCommand() = default;
+	void				Initialize(EntryFunc* InEntry, FStringView InName, FStringView InDesc, FArguments&& InArguments);
+	FContext			GetContext(int32 ArgC, const TCHAR* const* ArgV) const;
+	static FCommand*&	Head();
+	static int32		MainInner(int32 ArgC, const TCHAR* const* ArgV);
+	void				Usage() const;
+	int32				Call(int32 ArgC, const TCHAR* const* ArgV);
+	FCommand*			Next = FCommand::Head();
+	FStringView			Name;
+	FStringView			Desc;
+	FArguments			Arguments;
+	EntryFunc*			Entry;
+};
+
+} // namespace UE::IoStore::Tool
diff --git a/Engine/Source/Developer/IoStore/OnDemandUtilities/Private/Common.cpp b/Engine/Source/Developer/IoStore/OnDemandUtilities/Private/Common.cpp
new file mode 100644
--- /dev/null
+++ b/Engine/Source/Developer/IoStore/OnDemandUtilities/Private/Common.cpp
@@ -0,0 +1,27 @@
+// Copyright Epic Games, Inc. All Rights Reserved.
+
+#include "Common.h"
+#include "Command.h"
+
+#include "Misc/KeyChainUtilities.h"
+
+namespace UE::IoStore::Tool::Common
+{
+
+////////////////////////////////////////////////////////////////////////////////
+FKeyChain LoadCryptoKeys(const FContext& Context)
+{
+	FKeyChain Ret;
+
+	const FString Path = FString(Context.Get<FStringView>(TEXT("-CryptoKeys")));
+	if (Path.IsEmpty())
+	{
+		return Ret;
+	}
+
+	KeyChainUtilities::LoadKeyChainFromFile(*Path, Ret);
+
+	return Ret;
+}
+
+} // namespace UE::IoStore::Tool::Common
diff --git a/Engine/Source/Developer/IoStore/OnDemandUtilities/Private/Common.h b/Engine/Source/Developer/IoStore/OnDemandUtilities/Private/Common.h
new file mode 100644
--- /dev/null
+++ b/Engine/Source/Developer/IoStore/OnDemandUtilities/Private/Common.h
@@ -0,0 +1,14 @@
+// Copyright Epic Games, Inc. All Rights Reserved.
+
+#include "Containers/StringView.h"
+
+namespace UE::IoStore::Tool { class FContext; }
+struct FKeyChain;
+
+namespace UE::IoStore::Tool::Common
+{
+
+////////////////////////////////////////////////////////////////////////////////
+FKeyChain LoadCryptoKeys(const FContext& Context);
+
+} // namespace UE::IoStore::Tool::Common
diff --git a/Engine/Source/Developer/IoStore/OnDemandUtilities/Private/ConvertToc.cpp b/Engine/Source/Developer/IoStore/OnDemandUtilities/Private/ConvertToc.cpp
new file mode 100644
--- /dev/null
+++ b/Engine/Source/Developer/IoStore/OnDemandUtilities/Private/ConvertToc.cpp
@@ -0,0 +1,296 @@
+// Copyright Epic Games, Inc. All Rights Reserved.
+
+#include "Command.h"
+
+#include "HAL/FileManager.h"
+#include "IO/IoChunkId.h"
+#include "IO/IoContainerMeta.h"
+#include "IO/IoStoreOnDemand.h"
+#include "IO/OnDemandToc.h"
+#include "IO/PackageId.h"
+#include "IO/Serialization/OnDemandContainerToc.h"
+#include "Misc/FileHelper.h"
+#include "Misc/Paths.h"
+#include "Misc/StringBuilder.h"
+#include "Serialization/CompactBinarySerialization.h"
+#include "Serialization/CompactBinaryWriter.h"
+
+namespace UE::IoStore::Tool
+{
+
+////////////////////////////////////////////////////////////////////////////////
+static FIoChunkId GetPackageChunkId(const FIoChunkId& ChunkId, FPackageId& OutPackageId)
+{
+	FMemory::Memcpy(&OutPackageId, &ChunkId, sizeof(FPackageId));
+
+	if (ChunkId.GetChunkType() == EIoChunkType::ExportBundleData)
+	{
+		return ChunkId;
+	}
+
+	return CreatePackageDataChunkId(OutPackageId);
+}
+
+////////////////////////////////////////////////////////////////////////////////
+static FIoChunkId GetPackageChunkId(const FIoChunkId& ChunkId)
+{
+	FPackageId Unused;
+	return GetPackageChunkId(ChunkId, Unused);
+}
+
+////////////////////////////////////////////////////////////////////////////////
+class FContainerMeta
+{
+public:
+	FUtf8StringView			GetFilename(
+								const FIoChunkId& ChunkId,
+								FUtf8StringBuilderBase& OutFilename,
+								FUtf8StringView& OutContainerName);
+	FUtf8StringView			GetFilename(
+								const FIoChunkId& ChunkId,
+								FUtf8StringBuilderBase& OutFilename);
+	static FContainerMeta	Load(const FString& Directory);
+
+private:
+	TArray<FIoContainerMetaReader> Readers;
+};
+
+////////////////////////////////////////////////////////////////////////////////
+FUtf8StringView FContainerMeta::GetFilename(
+	const FIoChunkId& ChunkId,
+	FUtf8StringBuilderBase& OutFilename,
+	FUtf8StringView& OutContainerName)
+{
+	OutFilename.Reset();
+	for (FIoContainerMetaReader& Reader : Readers)
+	{
+		FUtf8StringView Filename = Reader.GetFilename(ChunkId, OutFilename, OutContainerName);
+		if (!Filename.IsEmpty())
+		{
+			return Filename;
+		}
+	}
+
+	return FUtf8StringView();
+}
+
+FUtf8StringView FContainerMeta::GetFilename(
+	const FIoChunkId& ChunkId,
+	FUtf8StringBuilderBase& OutFilename)
+{
+	FUtf8StringView Unused;
+	return GetFilename(ChunkId, OutFilename, Unused);
+}
+
+FContainerMeta FContainerMeta::Load(const FString& Directory)
+{
+	FContainerMeta	OutMeta;
+	TArray<FString> MetaFiles;
+
+	IFileManager& Ifm = IFileManager::Get();
+	Ifm.FindFiles(MetaFiles, *Directory, TEXT(".umeta"));
+
+	for (const FString& Filename : MetaFiles)
+	{
+		const FString	FullPath = Directory / Filename;
+		const bool		bMemoryMap = false;
+
+		if (TIoStatusOr<FIoContainerMetaReader> Status = FIoContainerMetaReader::Load(FullPath, bMemoryMap); Status.IsOk())
+		{
+			OutMeta.Readers.Add(Status.ConsumeValueOrDie());
+		}
+		else
+		{
+			UE_LOG(LogIoStoreOnDemand, Warning, TEXT("Failed to load container meta '%s'"), *FullPath);
+		}
+	}
+
+	return OutMeta;
+}
+
+////////////////////////////////////////////////////////////////////////////////
+static int32 ConvertTocCommandEntry(const FContext& Context)
+{
+	using namespace UE::IoStore::Serialization;
+	using namespace UE::IoStore::Serialization::V1;
+	using FChunkEntry = UE::IoStore::Serialization::V1::FOnDemandChunkEntry;
+
+	GPrintLogCategory = GPrintLogVerbosity = false;
+	GPrintLogTimes = ELogTimes::None;
+
+	FString FilePath = FString(Context.Get<FStringView>(TEXT("-Path")));
+	if (FilePath.IsEmpty())
+	{
+		UE_LOG(LogIoStoreOnDemand, Error, TEXT("Invalid filename"));
+		return -1;
+	}
+
+	FString OutFilePath = FString(Context.Get<FStringView>(TEXT("-Out")));
+	if (OutFilePath.IsEmpty())
+	{
+		UE_LOG(LogIoStoreOnDemand, Error, TEXT("Invalid JSON output filename"));
+		return -1;
+	}
+
+	const FString Ext = FPaths::GetExtension(OutFilePath).ToLower();
+	if (Ext != TEXT("json") && Ext != TEXT("csv"))
+	{
+		UE_LOG(LogIoStoreOnDemand, Error, TEXT("Invalid output file extension"));
+		return -1;
+	}
+
+	IFileManager& Ifm = IFileManager::Get();
+
+	TIoStatusOr<FOnDemandTocReader> MaybeReader = FOnDemandTocReader::Read(FilePath);
+	if (!MaybeReader.IsOk())
+	{
+		UE_LOG(LogIoStoreOnDemand, Error, TEXT("Failed to read on-demand TOC '%s', reason: %s"),
+			*FilePath, *MaybeReader.Status().ToString());
+		return -1;
+	}
+
+	FOnDemandTocReader	Reader = MaybeReader.ConsumeValueOrDie();
+	FContainerMeta		ContainerMeta = FContainerMeta::Load(FPaths::GetPath(FilePath));
+
+	if (Ext == TEXT("json"))
+	{
+		FCbWriter Writer;
+		Writer.BeginObject();
+		{
+			Writer.BeginObject(UTF8TEXTVIEW("Header"));
+			Writer.AddString(UTF8TEXTVIEW("CreationTime"), Reader.CreationTime().ToString());
+			Writer.AddString(UTF8TEXTVIEW("BuildVersion"), Reader.BuildVersion());
+			Writer.AddString(UTF8TEXTVIEW("TargetPlatform"), Reader.TargetPlatform());
+			Writer.AddString(UTF8TEXTVIEW("ChunksDirectory"), Reader.ChunksDirectory());
+			Writer.EndObject();
+
+			Writer.BeginArray(UTF8TEXTVIEW("Containers"));
+			for (const FOnDemandContainerEntry& ContainerEntry : Reader.Containers())
+			{
+				FOnDemandTocStorage Storage;
+				TIoStatusOr<FOnDemandContainerTocView> MaybeContainerView =
+					Reader.ReadContainer(ContainerEntry, Storage, EOnDemandTocReaderOptions::None);
+
+				if (MaybeContainerView.IsOk() == false)
+				{
+					UE_LOG(LogIoStoreOnDemand, Error, TEXT("Failed to read on-demand container TOC '%s', reason: %s"),
+						*FilePath, *MaybeContainerView.Status().ToString());
+					return -1;
+				}
+
+				FOnDemandContainerTocView ContainerView = MaybeContainerView.ConsumeValueOrDie();
+
+				Writer.BeginObject();
+				{
+					Writer.AddString(UTF8TEXTVIEW("ContainerName"), ContainerView.Header.ContainerName());
+					Writer.AddString(UTF8TEXTVIEW("ContainerId"), LexToString(ContainerView.Header.ContainerId()));
+					Writer.AddInteger(UTF8TEXTVIEW("ContainerFlags"), uint32(ContainerView.Header.ContainerFlags()));
+					Writer.AddString(UTF8TEXTVIEW("EncryptionKeyGuid"), ContainerView.Header.EncryptionKeyGuid().ToString());
+					Writer.AddInteger(UTF8TEXTVIEW("BlockSize"), ContainerView.Header.BlockSize());
+
+					TUtf8StringBuilder<512> Sb;
+					Writer.BeginArray(UTF8TEXTVIEW("ChunkEntries"));
+					for (int32 Idx = 0; const FIoChunkId& ChunkId : ContainerView.ChunkIds)
+					{
+						FUtf8StringView		Filename = ContainerMeta.GetFilename(GetPackageChunkId(ChunkId), Sb);
+						const FChunkEntry&	ChunkEntry = ContainerView.ChunkEntries[Idx++];
+						Writer.BeginObject();
+						{
+							Writer.AddString(UTF8TEXTVIEW("Filename"), Filename);
+							Writer.AddString(UTF8TEXTVIEW("ChunkId"), LexToString(ChunkId));
+							Writer.AddString(UTF8TEXTVIEW("Hash"), LexToString(ChunkEntry.Hash));
+							Writer.AddInteger(UTF8TEXTVIEW("RawSize"), ChunkEntry.RawSize);
+							Writer.AddInteger(UTF8TEXTVIEW("EncodedSize"), ChunkEntry.EncodedSize);
+							Writer.AddInteger(UTF8TEXTVIEW("BlockOffset"), ChunkEntry.BlockInfo.Offset());
+							Writer.AddInteger(UTF8TEXTVIEW("BlockCount"), ChunkEntry.BlockInfo.Count());
+							if (ChunkEntry.BlockInfo.Count() == 1)
+							{
+								Writer.AddInteger(UTF8TEXTVIEW("BlockSize"), ChunkEntry.BlockInfo.Size());
+								Writer.AddInteger(UTF8TEXTVIEW("BlockHash"), ChunkEntry.BlockInfo.Hash());
+							}
+						}
+						Writer.EndObject();
+					}
+					Writer.EndArray();
+				}
+				Writer.EndObject();
+			}
+			Writer.EndArray();
+		}
+		Writer.EndObject();
+
+		TStringBuilder<4096> Json;
+		CompactBinaryToJson(Writer.Save().AsObject(), Json);
+
+		if (!FFileHelper::SaveStringToFile(Json.ToView(), *OutFilePath, FFileHelper::EEncodingOptions::ForceUTF8WithoutBOM))
+		{
+			UE_LOG(LogIoStoreOnDemand, Error, TEXT("Failed to save JSON '%s'"), *OutFilePath);
+			return -1;
+		}
+	}
+	else
+	{
+		TUniquePtr<FArchive> Ar(Ifm.CreateFileWriter(*OutFilePath));
+		if (!Ar.IsValid())
+		{
+			UE_LOG(LogIoStoreOnDemand, Error, TEXT("Failed to open '%s' for writing"), *OutFilePath);
+			return -1;
+		}
+
+		Ar->Logf(TEXT("Platform, ContainerName, ChunkId, PackageId, Hash, Size, CompressedSize, Filename"));
+		for (const FOnDemandContainerEntry& ContainerEntry : Reader.Containers())
+		{
+			FOnDemandTocStorage Storage;
+			TIoStatusOr<FOnDemandContainerTocView> MaybeContainerView =
+				Reader.ReadContainer(ContainerEntry, Storage, EOnDemandTocReaderOptions::None);
+
+			if (MaybeContainerView.IsOk() == false)
+			{
+				UE_LOG(LogIoStoreOnDemand, Error, TEXT("Failed to read on-demand container TOC '%s', reason: %s"),
+					*FilePath, *MaybeContainerView.Status().ToString());
+				return -1;
+			}
+
+			FOnDemandContainerTocView ContainerView = MaybeContainerView.ConsumeValueOrDie();
+
+			TUtf8StringBuilder<512> Sb;
+			for (int32 Idx = 0; const FIoChunkId& ChunkId : ContainerView.ChunkIds)
+			{
+				FPackageId			PackageId;
+				FUtf8StringView		Filename = ContainerMeta.GetFilename(GetPackageChunkId(ChunkId, PackageId), Sb);
+				const FChunkEntry&	ChunkEntry = ContainerView.ChunkEntries[Idx++];
+				FString				ContainerFullName = FString(ContainerView.Header.ContainerName());
+
+				FString PakChunk, Platform;
+				ContainerFullName.Split(TEXT("-"), &PakChunk, &Platform, ESearchCase::CaseSensitive, ESearchDir::FromEnd);
+
+				Ar->Logf(TEXT("%s, %s, %s, %s, %s, %u, %u, %s"),
+					*Platform,
+					*PakChunk,
+					*LexToString(ChunkId),
+					*LexToString(PackageId),
+					*LexToString(ChunkEntry.Hash),
+					ChunkEntry.RawSize,
+					ChunkEntry.EncodedSize,
+					*FString(Filename));
+			}
+		}
+	}
+
+	UE_LOG(LogIoStoreOnDemand, Display, TEXT("Saved '%s'"), *OutFilePath);
+
+	return 0;
+}
+
+////////////////////////////////////////////////////////////////////////////////
+static FCommand ConvertTocCommand(
+	ConvertTocCommandEntry,
+	TEXT("ConvertToc"),
+	TEXT("Converts an on-demand TOC file to JSON or CSV"),
+	{
+		TArgument<FStringView>(TEXT("-Path"), TEXT("Path to a .uondemandtoc file")),
+		TArgument<FStringView>(TEXT("-Out"), TEXT("Output file path"))
+	}
+);
+
+} // namespace UE::IoStore::Tool
diff --git a/Engine/Source/Developer/IoStore/OnDemandUtilities/Private/Download.cpp b/Engine/Source/Developer/IoStore/OnDemandUtilities/Private/Download.cpp
new file mode 100644
--- /dev/null
+++ b/Engine/Source/Developer/IoStore/OnDemandUtilities/Private/Download.cpp
@@ -0,0 +1,387 @@
+// Copyright Epic Games, Inc. All Rights Reserved.
+
+#include "Command.h"
+
+#include "HAL/FileManager.h"
+#include "IO/IoStore.h"
+#include "IO/IoStoreOnDemand.h"
+#include "Misc/Paths.h"
+#include "S3/S3Client.h"
+#include "Serialization/CompactBinarySerialization.h"
+#include "Serialization/CompactBinaryWriter.h"
+#include "Serialization/LargeMemoryWriter.h"
+#include "Serialization/MemoryReader.h"
+
+namespace UE::IoStore::Tool
+{
+
+////////////////////////////////////////////////////////////////////////////////
+struct FDownloadParams
+{
+	FString Directory;
+	FString ServiceUrl;
+	FString Bucket;
+	FString Region; 
+	FString AccessKey;
+	FString SecretKey;
+	FString SessionToken;
+	FString CredentialsFile;
+	FString CredentialsFileKeyName;
+	int32 MaxConcurrentDownloads = 16;
+	
+	static TIoStatusOr<FDownloadParams> Parse(const TCHAR* CommandLine);
+	FIoStatus Validate() const;
+};
+
+FIoStatus FDownloadParams::Validate() const
+{
+	if (!AccessKey.IsEmpty() && SecretKey.IsEmpty())
+	{
+		return FIoStatus(EIoErrorCode::InvalidParameter, TEXT("Invalid secret key"));
+	}
+	else if (AccessKey.IsEmpty() && !SecretKey.IsEmpty())
+	{
+		return FIoStatus(EIoErrorCode::InvalidParameter, TEXT("Invalid access key"));
+	}
+
+	if (!CredentialsFile.IsEmpty() && CredentialsFileKeyName.IsEmpty())
+	{
+		return FIoStatus(EIoErrorCode::InvalidParameter, TEXT("Invalid credential file key name"));
+	}
+
+	if (ServiceUrl.IsEmpty() && Region.IsEmpty())
+	{
+		return FIoStatus(EIoErrorCode::InvalidParameter, TEXT("Service URL or AWS region needs to be specified"));
+	}
+
+	return FIoStatus::Ok;
+}
+
+////////////////////////////////////////////////////////////////////////////////
+FIoStatus DownloadContainerFiles(const FDownloadParams& DownloadParams, const FString& TocPath)
+{
+	struct FContainerStats
+	{
+		uint64 TocEntryCount = 0;
+		uint64 TocRawSize = 0;
+		uint64 RawSize = 0;
+		uint64 CompressedSize = 0;
+		EIoContainerFlags ContainerFlags = EIoContainerFlags::None;
+		FName CompressionMethod;
+	};
+	TMap<FString, FContainerStats> ContainerSummary;
+
+	const double StartTime = FPlatformTime::Seconds();
+
+	FS3ClientConfig Config;
+	Config.ServiceUrl = DownloadParams.ServiceUrl;
+	Config.Region = DownloadParams.Region;
+	
+	FS3ClientCredentials Credentials;
+	if (DownloadParams.CredentialsFile.IsEmpty() == false)
+	{
+		UE_LOG(LogIas, Display, TEXT("Loading credentials file '%s'"), *DownloadParams.CredentialsFile);
+		FS3CredentialsProfileStore CredentialsStore = FS3CredentialsProfileStore::FromFile(DownloadParams.CredentialsFile);
+		if (CredentialsStore.TryGetCredentials(DownloadParams.CredentialsFileKeyName, Credentials) == false)
+		{
+			return FIoStatus(EIoErrorCode::InvalidParameter, TEXT("Failed to find valid credentials in credentials file"));
+		}
+		else
+		{
+			UE_LOG(LogIas, Display, TEXT("Found credentials for '%s'"), *DownloadParams.CredentialsFileKeyName);
+		}
+	}
+	else
+	{
+		Credentials = FS3ClientCredentials(DownloadParams.AccessKey, DownloadParams.SecretKey, DownloadParams.SessionToken);
+	}
+
+	FS3Client Client(Config, Credentials);
+
+	UE_LOG(LogIas, Display, TEXT("Fetching TOC '%s/%s/%s'"), *Client.GetConfig().ServiceUrl, *DownloadParams.Bucket, *TocPath);
+	FS3GetObjectResponse TocResponse = Client.GetObject(FS3GetObjectRequest
+	{
+		DownloadParams.Bucket,
+		TocPath
+	});
+
+	if (TocResponse.IsOk() == false)
+	{
+		return FIoStatus(EIoErrorCode::ReadError, TEXT("Failed to fetch TOC"));
+	}
+
+	FOnDemandToc OnDemandToc;
+	FMemoryReaderView Ar(TocResponse.GetBody().GetView());
+	Ar << OnDemandToc;
+
+	if (Ar.IsError())
+	{
+		return FIoStatus(EIoErrorCode::ReadError, TEXT("Failed to load on demand TOC"));
+	}
+
+	FStringView BucketPrefix;
+	{
+		FStringView TocPathView = TocPath;
+		int32 Index = INDEX_NONE;
+		if (TocPathView.FindLastChar(TEXT('/'), Index))
+		{
+			BucketPrefix = TocPathView.Left(Index);
+		}
+	}
+
+	for (const FOnDemandTocContainerEntry& ContainerEntry : OnDemandToc.Containers)
+	{
+		TStringBuilder<256> FileTocKey;
+		if (BucketPrefix.IsEmpty() == false)
+		{
+			FileTocKey << BucketPrefix << TEXT("/");
+		}
+		FileTocKey << ContainerEntry.UTocHash << TEXT(".utoc");
+
+		UE_LOG(LogIas, Display, TEXT("Fetching '%s/%s/%s'"), *Client.GetConfig().ServiceUrl, *DownloadParams.Bucket, FileTocKey.ToString());
+		FS3GetObjectResponse Response = Client.GetObject(FS3GetObjectRequest
+		{
+			DownloadParams.Bucket,
+			FString(FileTocKey).ToLower()
+		});
+
+		if (Response.IsOk() == false)
+		{
+			return FIoStatus(EIoErrorCode::ReadError, TEXT("Failed to load container .utoc file"));
+		}
+
+		const FString UTocPath = FString::Printf(TEXT("%s/%s.utoc"), *DownloadParams.Directory, *ContainerEntry.ContainerName);
+		const FString UCasPath = FPaths::ChangeExtension(UTocPath, TEXT(".ucas")); 
+		FContainerStats& ContainerStats = ContainerSummary.FindOrAdd(ContainerEntry.ContainerName);
+		ContainerStats.TocRawSize = Response.GetBody().GetSize();
+
+		if (TUniquePtr<FArchive> TocFile(IFileManager::Get().CreateFileWriter(*UTocPath)); TocFile.IsValid())
+		{
+			UE_LOG(LogIas, Display, TEXT("Writing '%s'"), *UTocPath);
+			TocFile->Serialize((void*)Response.GetBody().GetData(), Response.GetBody().GetSize());
+		}
+		else
+		{
+			return FIoStatus(EIoErrorCode::WriteError, TEXT("Failed to write container .utoc file"));
+		}
+
+		FIoStoreTocResource FileToc;
+		if (FIoStatus Status = FIoStoreTocResource::Read(*UTocPath, EIoStoreTocReadOptions::ReadAll, FileToc); Status.IsOk() == false)
+		{
+			return Status;
+		}
+
+		const int32 TocEntryCount = int32(FileToc.Header.TocEntryCount);
+		const uint64 CompressionBlockSize = FileToc.Header.CompressionBlockSize;
+		ContainerStats.TocEntryCount = TocEntryCount;
+		ContainerStats.ContainerFlags = FileToc.Header.ContainerFlags; 
+
+		TMap<FIoChunkId, FIoHash> ChunkHashes;
+		for (const FOnDemandTocEntry& Entry : ContainerEntry.Entries)
+		{
+			ChunkHashes.Add(Entry.ChunkId, Entry.Hash);
+			ContainerStats.RawSize += Entry.RawSize;
+		}
+
+		TArray<int32> SortedIndices;
+		for (int32 Idx = 0; Idx < TocEntryCount; ++Idx)
+		{
+			SortedIndices.Add(Idx);
+		}
+
+		Algo::Sort(SortedIndices, [&FileToc](int32 Lhs, int32 Rhs)
+		{
+			return FileToc.ChunkOffsetLengths[Lhs].GetOffset() < FileToc.ChunkOffsetLengths[Rhs].GetOffset();
+		});
+
+		FString ChunksRelativePath = BucketPrefix.IsEmpty()
+			? TEXT("chunks")
+			: FString::Printf(TEXT("%s/chunks"), *FString(BucketPrefix));
+
+		TUniquePtr<FArchive> CasFile(IFileManager::Get().CreateFileWriter(*UCasPath));
+		TArray<uint8> PaddingBuffer;
+
+		for (int32 Idx = 0; Idx < TocEntryCount; ++Idx)
+		{
+			const int32 SortedIdx = SortedIndices[Idx];
+			const FIoChunkId& ChunkId = FileToc.ChunkIds[SortedIdx];
+			const FIoOffsetAndLength& OffsetLength = FileToc.ChunkOffsetLengths[SortedIdx];
+			const FIoHash ChunkHash = ChunkHashes.FindChecked(ChunkId);
+			const FString HashString = LexToString(ChunkHash);
+			const int32 FirstBlockIndex = int32(OffsetLength.GetOffset() / CompressionBlockSize);
+			const int32 LastBlockIndex = int32((Align(OffsetLength.GetOffset() + OffsetLength.GetLength(), CompressionBlockSize) - 1) / CompressionBlockSize);
+			const FIoStoreTocCompressedBlockEntry& FirstBlock = FileToc.CompressionBlocks[FirstBlockIndex];
+			const FName CompressionMethod = FileToc.CompressionMethods[FirstBlock.GetCompressionMethodIndex()];
+
+			if (CompressionMethod.IsNone() == false && ContainerStats.CompressionMethod.IsNone())
+			{
+				ContainerStats.CompressionMethod = CompressionMethod;
+			}
+
+			TStringBuilder<256> ChunkKey;
+			ChunkKey << ChunksRelativePath
+				<< TEXT("/") << HashString.Left(2)
+				<< TEXT("/") << HashString
+				<< TEXT(".iochunk");
+
+			UE_LOG(LogIas, Display, TEXT("Fetching '%s/%s/%s'"), *Client.GetConfig().ServiceUrl, *DownloadParams.Bucket, ChunkKey.ToString());
+			FS3GetObjectResponse ChunkResponse = Client.GetObject(FS3GetObjectRequest
+			{
+				DownloadParams.Bucket,
+				FString(ChunkKey).ToLower()
+			});
+
+			if (ChunkResponse.IsOk() == false)
+			{
+				return FIoStatus(EIoErrorCode::ReadError, TEXT("Failed to fetch chunk"));
+			}
+
+			const uint64 Padding = FirstBlock.GetOffset() - CasFile->Tell();
+			if (Padding > PaddingBuffer.Num())
+			{
+				PaddingBuffer.SetNumZeroed(int32(Padding));
+			}
+
+			if (Padding > 0)
+			{
+				CasFile->Serialize(PaddingBuffer.GetData(), Padding);
+			}
+
+			UE_LOG(LogIas, Display, TEXT("Serializing chunk %d/%d '%s' -> '%s' (%llu B)"),
+				Idx + 1, TocEntryCount, *HashString, *LexToString(ChunkId), ChunkResponse.GetBody().GetSize());
+
+			check(CasFile->Tell() == FirstBlock.GetOffset());
+			CasFile->Serialize((void*)ChunkResponse.GetBody().GetData(), ChunkResponse.GetBody().GetSize());
+		}
+
+		ContainerStats.CompressedSize = CasFile->Tell();
+	}
+
+	{
+		const double Duration = FPlatformTime::Seconds() - StartTime;
+
+		UE_LOG(LogIas, Display, TEXT(""));
+		UE_LOG(LogIas, Display, TEXT("---------------------------------------- Download Summary --------------------------------------"));
+		UE_LOG(LogIas, Display, TEXT("%-40s: %s"), TEXT("Service URL"), *DownloadParams.ServiceUrl);
+		UE_LOG(LogIas, Display, TEXT("%-40s: %s"), TEXT("Bucket"), *DownloadParams.Bucket);
+		UE_LOG(LogIas, Display, TEXT("%-40s: %s"), TEXT("TOC"), *TocPath);
+		UE_LOG(LogIas, Display, TEXT("%-40s: %.2lf second(s)"), TEXT("Duration"), Duration);
+		UE_LOG(LogIas, Display, TEXT(""));
+
+		UE_LOG(LogIas, Display, TEXT("%-30s %10s %15s %15s %15s %25s"),
+			TEXT("Container"), TEXT("Flags"), TEXT("TOC Size (KB)"), TEXT("TOC Entries"), TEXT("Size (MB)"), TEXT("Compressed (MB)"));
+		UE_LOG(LogIas, Display, TEXT("-------------------------------------------------------------------------------------------------------------------------"));
+		
+		FContainerStats TotalStats;
+		for (const TTuple<FString, FContainerStats>& Kv : ContainerSummary)
+		{
+			const FString& ContainerName = Kv.Key;
+			const FContainerStats& Stats = Kv.Value;
+
+			FString CompressionInfo = "-";
+
+			if (Stats.CompressionMethod != NAME_None)
+			{
+				const double Percentage = Stats.RawSize > 0 ? (double(Stats.RawSize - Stats.CompressedSize) / double(Stats.RawSize)) * 100.0 : 0.0;
+				CompressionInfo = FString::Printf(TEXT("%.2lf (%.2lf%% %s)"),
+					(double)Stats.CompressedSize / 1024.0 / 1024.0,
+					Percentage,
+					*Stats.CompressionMethod.ToString());
+			}
+
+			FString ContainerSettings = FString::Printf(TEXT("%s/%s/%s/%s/%s"),
+				EnumHasAnyFlags(Stats.ContainerFlags, EIoContainerFlags::Compressed) ? TEXT("C") : TEXT("-"),
+				EnumHasAnyFlags(Stats.ContainerFlags, EIoContainerFlags::Encrypted) ? TEXT("E") : TEXT("-"),
+				EnumHasAnyFlags(Stats.ContainerFlags, EIoContainerFlags::Signed) ? TEXT("S") : TEXT("-"),
+				EnumHasAnyFlags(Stats.ContainerFlags, EIoContainerFlags::Indexed) ? TEXT("I") : TEXT("-"),
+				EnumHasAnyFlags(Stats.ContainerFlags, EIoContainerFlags::OnDemand) ? TEXT("O") : TEXT("-"));
+
+			UE_LOG(LogIas, Display, TEXT("%-30s %10s %15.2lf %15llu %15.2lf %25s"),
+				*ContainerName,
+				*ContainerSettings,
+				(double)Stats.TocRawSize / 1024.0,
+				Stats.TocEntryCount,
+				(double)Stats.RawSize / 1024.0 / 1024.0,
+				*CompressionInfo);
+
+			TotalStats.TocEntryCount += Stats.TocEntryCount;
+			TotalStats.TocRawSize += Stats.TocRawSize;
+			TotalStats.RawSize += Stats.RawSize;
+			TotalStats.CompressedSize += Stats.CompressedSize;
+		}
+		UE_LOG(LogIas, Display, TEXT("-------------------------------------------------------------------------------------------------------------------------"));
+		UE_LOG(LogIas, Display, TEXT("%-30s %10s %15.2lf %15llu %15.2lf %25.2lf "),
+			TEXT("Total"),
+			TEXT(""),
+			(double)TotalStats.TocRawSize / 1024.0,
+			TotalStats.TocEntryCount,
+			(double)TotalStats.RawSize / 1024.0 / 1024.0,
+			(double)TotalStats.CompressedSize / 1024.0 / 1024.0);
+
+		UE_LOG(LogIas, Display, TEXT(""));
+		UE_LOG(LogIas, Display, TEXT("** Flags: (C)ompressed / (E)ncrypted / (S)igned) / (I)ndexed) / (O)nDemand **"));
+		UE_LOG(LogIas, Display, TEXT(""));
+	}
+
+	return FIoStatus::Ok;
+}
+
+////////////////////////////////////////////////////////////////////////////////
+static FDownloadParams BuildDownloadParams(const FContext& Context)
+{
+	FDownloadParams Ret;
+
+	Ret.Directory				= Context.Get<FStringView>(TEXT("-Directory"),				Ret.Directory);
+	Ret.ServiceUrl				= Context.Get<FStringView>(TEXT("-ServiceUrl"),				Ret.ServiceUrl);
+	Ret.Bucket					= Context.Get<FStringView>(TEXT("-Bucket"),					Ret.Bucket);
+	Ret.Region					= Context.Get<FStringView>(TEXT("-Region"),					Ret.Region);
+	Ret.AccessKey				= Context.Get<FStringView>(TEXT("-AccessKey"),				Ret.AccessKey);
+	Ret.SecretKey				= Context.Get<FStringView>(TEXT("-SecretKey"),				Ret.SecretKey);
+	Ret.SessionToken			= Context.Get<FStringView>(TEXT("-SessionToken"),			Ret.SessionToken);
+	Ret.CredentialsFile			= Context.Get<FStringView>(TEXT("-CredentialsFile"),		Ret.CredentialsFile);
+	Ret.CredentialsFileKeyName	= Context.Get<FStringView>(TEXT("-CredentialsFileKeyName"),	Ret.CredentialsFileKeyName);
+
+	Ret.MaxConcurrentDownloads	= Context.Get(TEXT("-MaxConcurrentDownloads"),	Ret.MaxConcurrentDownloads);
+
+	if (FIoStatus Validation = Ret.Validate(); !Validation.IsOk())
+	{
+		FString Reason = Validation.ToString();
+		Context.Abort(*Reason);
+	}
+
+	return Ret;
+}
+
+////////////////////////////////////////////////////////////////////////////////
+static int32 DownloadCommandEntry(const Tool::FContext& Context)
+{
+	FStringView TocPath = Context.Get<FStringView>(TEXT("TocPath"));
+
+	FDownloadParams Params = BuildDownloadParams(Context);
+
+	FIoStatus Status = DownloadContainerFiles(Params, FString(TocPath));
+	if (!Status.IsOk())
+	{
+		FString Reason = Status.ToString();
+		Context.Abort(*Reason);
+	}
+
+	return 0;
+}
+
+////////////////////////////////////////////////////////////////////////////////
+extern FArgumentSet S3Arguments;
+
+static FCommand DownloadCommand(
+	DownloadCommandEntry,
+	TEXT("Download"),
+	TEXT("Fetches the cloud-stored contents of a given on-demand TOC"),
+	{
+		TArgument<FStringView>(TEXT("TocPath"),				TEXT("Bucket-relative path of the TOC to download")),
+		TArgument<FStringView>(TEXT("-Directory"),			TEXT("Output directory")),
+		TArgument<int32>(TEXT("-MaxConcurrentDownloads"),	TEXT("Number of downloads that happen all at once")),
+		S3Arguments,
+	}
+);
+
+} // namespace UE::IoStore::Tool
diff --git a/Engine/Source/Developer/IoStore/OnDemandUtilities/Private/Graph.cpp b/Engine/Source/Developer/IoStore/OnDemandUtilities/Private/Graph.cpp
new file mode 100644
--- /dev/null
+++ b/Engine/Source/Developer/IoStore/OnDemandUtilities/Private/Graph.cpp
@@ -0,0 +1,1129 @@
+// Copyright Epic Games, Inc. All Rights Reserved.
+
+#include "Graph.h"
+#include "Command.h"
+#include "Common.h"
+
+#include "HAL/FileManager.h"
+#include "HAL/Platform.h"
+#include "IO/IoContainerHeader.h"
+#include "IO/IoStore.h"
+#include "IO/IoStoreOnDemand.h"
+#include "IO/PackageId.h"
+#include "Misc/CString.h"
+#include "Misc/KeyChainUtilities.h"
+#include "Misc/Paths.h"
+#include "Misc/Parse.h"
+#include "Serialization/CompactBinary.h"
+#include "Serialization/CompactBinaryWriter.h"
+#include "Serialization/CompactBinarySerialization.h"
+#include "Serialization/LargeMemoryReader.h"
+
+namespace UE::IoStore::Tool
+{
+
+////////////////////////////////////////////////////////////////////////////////
+FIoStatus ExportGephiCsv(const FPackageGraph& Graph, const FString& OutFilename)
+{
+	UE_LOG(LogIoStoreOnDemand, Display, TEXT("Exporting CSV with %d nodes and %d edges"), Graph.GetNodeCount(), Graph.GetEdgeCount());
+	IFileManager& Ifm = IFileManager::Get();
+	{
+		TUniquePtr<FArchive> Ar(Ifm.CreateFileWriter(*OutFilename));
+		if (!Ar.IsValid())
+		{
+			return FIoStatus(EIoErrorCode::WriteError);
+		}
+
+		UE_LOG(LogIoStoreOnDemand, Display, TEXT("Saving '%s'"), *OutFilename);
+		Ar->Logf(TEXT("Id, Label"));
+		for (int32 Idx = 0; Idx < Graph.GetNodeCount(); ++Idx)
+		{
+			FPackageGraph::FNode			Node{Idx};
+			const FPackageGraph::FNodeData&	Data = Graph.GetNodeData(Node);
+			Ar->Logf(TEXT("%d, %s"), Idx, *LexToString(Data.PackageId));
+		}
+
+		if (!Ar->Close())
+		{
+			return FIoStatus(EIoErrorCode::WriteError);
+		}
+	}
+
+	{
+		FString Filename		= FPaths::GetBaseFilename(OutFilename);
+		FString Path			= FPaths::GetPath(OutFilename);
+		FString EdgesFilename	= Path / Filename + TEXT("-edges.csv");
+
+		TUniquePtr<FArchive> Ar(Ifm.CreateFileWriter(*EdgesFilename));
+		if (!Ar.IsValid())
+		{
+			return FIoStatus(EIoErrorCode::WriteError);
+		}
+
+		UE_LOG(LogIoStoreOnDemand, Display, TEXT("Saving '%s'"), *EdgesFilename);
+		Ar->Logf(TEXT("Source, Target, Type, Weight, Id"));
+		for (int32 Idx = 0; Idx < Graph.GetEdgeCount(); ++Idx)
+		{
+			const FPackageGraph::FEdgeData&	Edge	= Graph.GetEdgeData(FPackageGraph::FEdge { .Index = Idx });
+			const float						Weight	= 1.0f;
+			const TCHAR*					Type	= TEXT("Directed");
+			Ar->Logf(TEXT("%d, %d, %s, %.2f, %d"), Edge.From.Index, Edge.To.Index, Type, Weight, Idx);
+		}
+
+		if (!Ar->Close())
+		{
+			return FIoStatus(EIoErrorCode::WriteError);
+		}
+	}
+
+	return FIoStatus::Ok;
+}
+
+////////////////////////////////////////////////////////////////////////////////
+void FPackageGraph::Reserve(int32 ContainerCount, int32 PackageCount)
+{
+	ContainerNodeData.Reserve(ContainerCount);
+	NodeData.Reserve(PackageCount);
+	EdgeData.Reserve(PackageCount);
+	PackageIdToNode.Reserve(PackageCount);
+}
+
+FPackageGraph::FContainerNode FPackageGraph::AddContainerNode(FString&& Name, const FIoContainerId& Id)
+{
+	const int32 Index = ContainerNodeData.Num();
+	ContainerNodeData.Add(FContainerData
+	{
+		.Name	= MoveTemp(Name),
+		.Id		= Id
+	});
+
+	return FContainerNode { .Index = Index };
+}
+
+FPackageGraph::FNode FPackageGraph::AddNode(
+	FPackageGraph::FContainerNode ContainerNode,
+	const FPackageId& PackageId,
+	const FString* Filename,
+	uint64 ChunkSize)
+{
+	check(ContainerNode.IsValid());
+	if (FNode Node = PackageIdToNode.FindRef(PackageId); Node.IsValid())
+	{
+		return Node;
+	}
+
+	const int32 Index	= NodeData.Num();
+	FNodeData& Data		= NodeData.AddDefaulted_GetRef();
+	Data.PackageId		= PackageId;
+	Data.Size			= ChunkSize;
+	Data.ContainerNode	= ContainerNode;
+
+	if (Filename != nullptr)
+	{
+		Data.Filename	= *Filename;
+	}
+
+	FNode NewNode{Index};
+	PackageIdToNode.Add(PackageId, NewNode);
+	ContainerNodeData[ContainerNode.Index].PackageNodes.Add(NewNode);
+
+	return NewNode;
+}
+
+void FPackageGraph::AddEdge(FNode From, FNode To, bool bSoft)
+{
+	check(From.IsValid() && To.IsValid());
+
+	const FEdge Edge { .Index = EdgeData.Num() };
+	EdgeData.Add(FEdgeData { .From = From, .To = To, .bSoft = bSoft });
+
+	FNodeData& FromNodeData = NodeData[From.Index];
+	FNodeData& ToNodeData	= NodeData[To.Index];
+
+	FromNodeData.Connections.Out.Add(Edge);
+	ToNodeData.Connections.In.Add(Edge);
+}
+
+void FPackageGraph::AddMissingEdge(FNode From, const FPackageId To, bool bSoft)
+{
+	check(From.IsValid()); 
+
+	const FMissingEdge MissingEdge { .Index = MissingEdgeData.Num() };
+	MissingEdgeData.Add(FMissingEdgeData { .From = From, .To = To, .bSoft = bSoft});
+
+	FNodeData& FromNodeData = NodeData[From.Index];
+}
+
+void FPackageGraph::AddRedirect(const FPackageId& From, const FPackageId& To)
+{
+	Redirects.Add(FRedirect
+	{
+		.From	= From,
+		.To		= To
+	});
+}
+
+void FPackageGraph::Traverse(FNode StartNode, bool bIncoming, bool bSkipSoftRefs, FVisitor& Visitor)
+{
+	TSet<FNode>		Visited;
+	TArray<FNode>	Stack;
+
+	check(StartNode.IsValid());
+
+	Stack.Push(StartNode);
+	while (Stack.IsEmpty() == false)
+	{
+		const FNode Node = Stack.Pop();
+
+		bool bAlreadyInSet = false;
+		Visited.Add(Node, &bAlreadyInSet);
+		if (bAlreadyInSet)
+		{
+			continue;
+		}
+
+		const FNodeData& Data = GetNodeData(Node);
+		if (Visitor.Visit(*this, Node, Data) == false)
+		{
+			return;
+		}
+
+		const TSet<FEdge>& Edges = bIncoming ? Data.Connections.In : Data.Connections.Out;
+		for (FEdge Edge : Edges)
+		{
+			const FEdgeData& ToTraverseEdgeData = EdgeData[Edge.Index];
+			if (ToTraverseEdgeData.bSoft && bSkipSoftRefs)
+			{
+				continue;
+			}
+
+			const FNode ToTraverse = bIncoming ? ToTraverseEdgeData.From : ToTraverseEdgeData.To;
+			if (Visited.Contains(ToTraverse) == false)
+			{
+				Stack.Push(ToTraverse);
+			}
+		}
+	}
+}
+
+void FPackageGraph::Build(
+	TConstArrayView<FString> ContainerFilenames,
+	const FKeyChain& KeyChain,
+	bool bIncludeSoftRefs,
+	FPackageGraph& Graph)
+{
+	UE_LOG(LogIoStoreOnDemand, Display, TEXT("Building package graph..."));
+
+	TMap<FGuid, FAES::FAESKey> EncryptionKeys;
+	for (const TPair<FGuid, FNamedAESKey>& KeyPair: KeyChain.GetEncryptionKeys())
+	{
+		EncryptionKeys.Add(KeyPair.Key, KeyPair.Value.Key);
+	}
+
+	struct FLoadedContainer
+	{
+		FString						Filename;
+		FString						Name;
+		FIoStoreReader				Reader;
+		FIoContainerHeader			Header;
+		TMap<FIoChunkId, FString>	ChunkFileNamesMap;
+	};
+	TArray<TUniquePtr<FLoadedContainer>> Containers;
+	Containers.Reserve(ContainerFilenames.Num());
+	int32 TotalPackages = 0;
+
+	// Load all containers
+	for (const FString& Filename : ContainerFilenames)
+	{
+		// Only .utoc file(s) 
+		if (FPaths::GetExtension(Filename, false) != TEXT("utoc"))
+		{
+			continue;
+		}
+
+		// Skip .o.utoc
+		if (FPaths::GetExtension(FPaths::GetBaseFilename(Filename), false) == TEXT("o"))
+		{
+			continue;
+		}
+
+		TUniquePtr<FLoadedContainer>	Container = MakeUnique<FLoadedContainer>();
+		FString							Path = Filename;
+
+		FIoStatus Status = Container->Reader.Initialize(*FPaths::ChangeExtension(Path, TEXT("")), EncryptionKeys);
+		if (!Status.IsOk())
+		{
+			UE_LOG(LogIoStoreOnDemand, Error, TEXT("Failed to open container '%s' for reading due to: '%s'"), *Path, *Status.ToString());
+			continue;
+		}
+
+		const FIoChunkId		HeaderChunkId = CreateContainerHeaderChunkId(Container->Reader.GetContainerId());
+		TIoStatusOr<FIoBuffer>	Chunk = Container->Reader.Read(HeaderChunkId, FIoReadOptions());
+
+		if (Chunk.IsOk() == false)
+		{
+			// Container header isn't required 
+			continue;
+		}
+
+		FIoBuffer			Buffer = Chunk.ValueOrDie();
+		FLargeMemoryReader	Ar(Buffer.GetData(), Buffer.GetSize());
+
+		Ar << Container->Header;
+		if (Ar.IsError() || Ar.IsCriticalError())
+		{
+			UE_LOG(LogIoStoreOnDemand, Error, TEXT("Failed to serialize container header, filename '%s'"), *Filename);
+			continue;
+		}
+
+		if (Container->Header.SoftPackageReferencesSerialInfo.Size > 0)
+		{
+			Ar.Seek(Container->Header.SoftPackageReferencesSerialInfo.Offset);
+			Ar << Container->Header.SoftPackageReferences;
+		}
+
+		if (Ar.IsError() || Ar.IsCriticalError())
+		{
+			UE_LOG(LogIoStoreOnDemand, Error, TEXT("Failed to serialize container soft references, filename '%s'"), *Filename);
+			continue;
+		}
+
+		for (const FIoContainerHeaderPackageRedirect& Redirect : Container->Header.PackageRedirects)
+		{
+			Graph.AddRedirect(Redirect.SourcePackageId, Redirect.TargetPackageId); 
+		}
+
+		Container->Reader.GetDirectoryIndexReader().IterateDirectoryIndex(FIoDirectoryIndexHandle::RootDirectory(), TEXT(""),
+			[&Container](FStringView Filename, uint32 TocEntryIndex) -> bool
+			{
+				TIoStatusOr<FIoStoreTocChunkInfo> ChunkInfo = Container->Reader.GetChunkInfo(TocEntryIndex);
+				if (ChunkInfo.IsOk())
+				{
+					Container->ChunkFileNamesMap.Add(ChunkInfo.ValueOrDie().Id, FString(Filename));
+				}
+				return true;
+			});
+
+
+		TotalPackages += Container->Header.PackageIds.Num();
+		Container->Name = FPaths::GetBaseFilename(Filename);
+		Containers.Add(MoveTemp(Container));
+	}
+
+	Graph.Reserve(Containers.Num(), TotalPackages);
+
+	// Create graph nodes 
+	for (const TUniquePtr<FLoadedContainer>& Container : Containers)
+	{
+		const FString&		Filename = Container->Filename;
+		FIoStoreReader&		Reader = Container->Reader;
+		FIoContainerHeader&	Header = Container->Header;
+
+		FString Name = Container->Name;
+		const FPackageGraph::FContainerNode ContainerNode = Graph.AddContainerNode(MoveTemp(Name), Reader.GetContainerId());
+		const int32							NodeCount = Graph.GetNodeCount();
+
+		for (int32 Idx = 0; Idx < Header.PackageIds.Num(); ++Idx)
+		{
+			const FPackageId	PackageId = Header.PackageIds[Idx];
+			const FIoChunkId	ChunkId = CreatePackageDataChunkId(PackageId);
+			const FString*		ChunkFilename = Container->ChunkFileNamesMap.Find(ChunkId);
+
+			TIoStatusOr<FIoStoreTocChunkInfo> ChunkInfo = Reader.GetChunkInfo(ChunkId);
+			if (ChunkInfo.IsOk() == false)
+			{
+				UE_LOG(LogIoStoreOnDemand, Error, TEXT("Failed to find chunk info for package '%s'"), *LexToString(PackageId));
+				continue;
+			}
+
+			Graph.AddNode(ContainerNode, PackageId, ChunkFilename, Align(ChunkInfo.ValueOrDie().CompressedSize, 16));
+		}
+
+		UE_LOG(LogIoStoreOnDemand, Display, TEXT("Created %d node(s) from '%s'"), Graph.GetNodeCount() - NodeCount, *Container->Name);
+	}
+
+	// Create graph edges
+	for (const TUniquePtr<FLoadedContainer>& Container : Containers)
+	{
+		const FString&		Filename = Container->Filename;
+		FIoStoreReader&		Reader = Container->Reader;
+		FIoContainerHeader&	Header = Container->Header;
+		const int32			EdgeCount = Graph.GetEdgeCount();
+
+		TConstArrayView<FFilePackageStoreEntry> PackageStoreEntries = MakeArrayView(
+			reinterpret_cast<const FFilePackageStoreEntry*>(Header.StoreEntries.GetData()), Header.PackageIds.Num());
+
+		TConstArrayView<FFilePackageStoreEntrySoftReferences> AllSoftReferences;
+		if (Header.SoftPackageReferences.bContainsSoftPackageReferences)
+		{
+			AllSoftReferences = MakeArrayView<const FFilePackageStoreEntrySoftReferences>(
+				reinterpret_cast<const FFilePackageStoreEntrySoftReferences*>(Header.SoftPackageReferences.PackageIndices.GetData()),
+				Header.PackageIds.Num());
+		}
+
+		for (int32 Idx = 0; Idx < Header.PackageIds.Num(); ++Idx)
+		{
+			const FPackageId PackageId			= Header.PackageIds[Idx];
+			const FFilePackageStoreEntry& Entry = PackageStoreEntries[Idx];
+
+			// Hard reference(s)
+			FPackageGraph::FNode FromNode = Graph.GetNode(PackageId);
+			check(FromNode.IsValid());
+			for (const FPackageId& ImportedPackageId : Entry.ImportedPackages)
+			{
+				const bool					bSoftEdge = false;
+				const FPackageGraph::FNode	ToNode = Graph.GetNode(ImportedPackageId);
+				if (ToNode.IsValid())
+				{
+					//TODO: Dedup edges
+					Graph.AddEdge(FromNode, ToNode, bSoftEdge);
+				}
+				else
+				{
+					Graph.AddMissingEdge(FromNode, ImportedPackageId, bSoftEdge); 
+				}
+			}
+
+			// Soft reference(s)
+			if (bIncludeSoftRefs && AllSoftReferences.IsEmpty() == false)
+			{
+				const FFilePackageStoreEntrySoftReferences& SoftRefs = AllSoftReferences[Idx];
+				if (SoftRefs.Indices.Num() > 0)
+				{
+					for (int32 SoftRefIdx : SoftRefs.Indices)
+					{
+						const bool					bSoftEdge = true;
+						const FPackageId			SoftRef = Header.SoftPackageReferences.PackageIds[SoftRefIdx];
+						const FPackageGraph::FNode	ToNode = Graph.GetNode(SoftRef);
+						if (ToNode.IsValid())
+						{
+							Graph.AddEdge(FromNode, ToNode, bSoftEdge);
+						}
+						else
+						{
+							Graph.AddMissingEdge(FromNode, SoftRef, bSoftEdge);
+						}
+					}
+				}
+			}
+		}
+
+		UE_LOG(LogIoStoreOnDemand, Display, TEXT("Created %d edges(s) from '%s'"), Graph.GetEdgeCount() - EdgeCount, *Container->Name);
+	}
+}
+
+FIoStatus FPackageGraph::Save(const FPackageGraph& Graph, const FString& Filename)
+{
+	const FString Ext = FPaths::GetExtension(Filename);
+
+	if (Ext == TEXT("json") || Ext == TEXT("cb"))
+	{
+		UE_LOG(LogIoStoreOnDemand, Display, TEXT("Saving graph to '%s'"), *Filename);
+
+		IFileManager&	Ifm = IFileManager::Get();
+		FCbWriter		Writer;
+
+		Writer << Graph;
+
+		TUniquePtr<FArchive> Ar = TUniquePtr<FArchive>(Ifm.CreateFileWriter(*Filename));
+		if (Ar.IsValid())
+		{
+			if (Ext == TEXT("cb"))
+			{
+				Writer.Save(*Ar);
+			}
+			else
+			{
+				FCbObject GraphCb = Writer.Save().AsObject();
+				FUtf8StringBuilderBase Sb;
+				CompactBinaryToJson(GraphCb, Sb);
+				UTF8CHAR UTF8BOM[] = { (UTF8CHAR)0xEF, (UTF8CHAR)0xBB, (UTF8CHAR)0xBF };
+				Ar->Serialize(&UTF8BOM, UE_ARRAY_COUNT(UTF8BOM) * sizeof(UTF8CHAR));
+				Ar->Serialize(Sb.GetData(), Sb.Len() * sizeof(UTF8CHAR));
+			}
+			Ar->Close();
+
+			if (Ar->IsError() || Ar->IsCriticalError())
+			{
+				return FIoStatus(EIoErrorCode::WriteError);
+			}
+
+			return FIoStatus::Ok;
+		}
+		else
+		{
+			return FIoStatus(EIoErrorCode::FileOpenFailed);
+		}
+	}
+	else if (Ext == TEXT("csv"))
+	{
+		return ExportGephiCsv(Graph, Filename);
+	}
+
+	return FIoStatus(EIoErrorCode::InvalidParameter, TEXT("Unsupported file format"));
+}
+
+////////////////////////////////////////////////////////////////////////////////
+static bool LoadFromCompactBinary(FCbFieldView Field, FPackageGraph::FIndex& Index)
+{
+	if (Field.IsInteger())
+	{
+		Index.Index = Field.AsInt32();
+		return true;
+	}
+
+	return false;
+}
+
+////////////////////////////////////////////////////////////////////////////////
+static FCbWriter& operator<<(FCbWriter& Writer, const FPackageGraph::FIndex& Index)
+{
+	Writer.AddInteger(Index.Index);
+	return Writer;
+}
+
+////////////////////////////////////////////////////////////////////////////////
+static bool LoadFromCompactBinary(FCbFieldView Field, FPackageGraph::FEdgeData& Edge)
+{
+	if (FCbObjectView Obj = Field.AsObjectView())
+	{
+		Edge.From	= FPackageGraph::FNode { .Index = Obj["From"].AsInt32() };
+		Edge.To		= FPackageGraph::FNode { .Index = Obj["To"].AsInt32() };
+		Edge.bSoft	= Obj["Soft"].AsBool();
+
+		return true;
+	}
+
+	return false;
+}
+
+////////////////////////////////////////////////////////////////////////////////
+static FCbWriter& operator<<(FCbWriter& Writer, const FPackageGraph::FEdgeData& Edge)
+{
+	Writer.BeginObject();
+	Writer << UTF8TEXT("From") << Edge.From;
+	Writer << UTF8TEXT("To") << Edge.To;
+	Writer << UTF8TEXT("Soft") << Edge.bSoft;
+	Writer.EndObject();
+
+	return Writer;
+}
+
+////////////////////////////////////////////////////////////////////////////////
+static bool LoadFromCompactBinary(FCbFieldView Field, FPackageGraph::FNodeData& Node)
+{
+	FCbObjectView Obj = Field.AsObjectView();
+
+	if (!Obj)
+	{
+		return false;
+	}
+
+	Node.Filename		= FString(Obj["Filename"].AsString());
+	Node.PackageId		= FPackageId::FromValue(Obj["PackageId"].AsUInt64());
+	Node.Size			= Obj["Size"].AsUInt64();
+	Node.ContainerNode	= FPackageGraph::FContainerNode { .Index = Obj["Container"].AsInt32() };
+
+	if (FCbArrayView ArrayView = Obj["In"].AsArrayView())
+	{
+		Node.Connections.In.Reserve(int32(ArrayView.Num()));
+		for (FCbFieldView ArrayField : ArrayView)
+		{
+			Node.Connections.In.Add(FPackageGraph::FEdge { .Index = ArrayField.AsInt32() });
+		}
+	}
+
+	if (FCbArrayView ArrayView = Obj["Out"].AsArrayView())
+	{
+		Node.Connections.Out.Reserve(int32(ArrayView.Num()));
+		for (FCbFieldView ArrayField : ArrayView)
+		{
+			Node.Connections.Out.Add(FPackageGraph::FEdge { .Index = ArrayField.AsInt32() });
+		}
+	}
+
+	return true;
+}
+
+////////////////////////////////////////////////////////////////////////////////
+static FCbWriter& operator<<(FCbWriter& Writer, const FPackageGraph::FNodeData& Data)
+{
+	Writer.BeginObject();
+	Writer << UTF8TEXT("Filename") << Data.Filename;
+	Writer << UTF8TEXT("PackageId") << Data.PackageId.Value();
+	Writer << UTF8TEXT("Size") << Data.Size;
+	Writer << UTF8TEXT("Container") << Data.ContainerNode;
+
+	Writer.BeginArray(UTF8TEXT("In"));
+	for (FPackageGraph::FEdge Edge : Data.Connections.In)
+	{
+		Writer << Edge;
+	}
+	Writer.EndArray();
+
+	Writer.BeginArray(UTF8TEXT("Out"));
+	for (FPackageGraph::FEdge Edge : Data.Connections.Out)
+	{
+		Writer << Edge;
+	}
+	Writer.EndArray();
+
+	Writer.EndObject();
+	return Writer;
+}
+
+////////////////////////////////////////////////////////////////////////////////
+static bool LoadFromCompactBinary(FCbFieldView Field, FPackageGraph::FContainerData& Container)
+{
+	FCbObjectView Obj = Field.AsObjectView();
+
+	if (!Obj)
+	{
+		return false;
+	}
+
+	Container.Name = FString(Obj["Name"].AsString());
+	if (!LoadFromCompactBinary(Obj["Id"], Container.Id))
+	{
+		return false;
+	}
+
+	if (FCbArrayView ArrayView = Obj["In"].AsArrayView())
+	{
+		Container.Connections.In.Reserve(int32(ArrayView.Num()));
+		for (FCbFieldView ArrayField : ArrayView)
+		{
+			Container.Connections.In.Add(FPackageGraph::FEdge { .Index = ArrayField.AsInt32() });
+		}
+	}
+
+	if (FCbArrayView ArrayView = Obj["Out"].AsArrayView())
+	{
+		Container.Connections.Out.Reserve(int32(ArrayView.Num()));
+		for (FCbFieldView ArrayField : ArrayView)
+		{
+			Container.Connections.Out.Add(FPackageGraph::FEdge { .Index = ArrayField.AsInt32() });
+		}
+	}
+
+	return true;
+}
+
+////////////////////////////////////////////////////////////////////////////////
+static FCbWriter& operator<<(FCbWriter& Writer, const FPackageGraph::FContainerData& Container)
+{
+	Writer.BeginObject();
+	Writer << UTF8TEXT("Name") << Container.Name;
+	Writer << UTF8TEXT("Id") << Container.Id;
+
+	Writer.BeginArray(UTF8TEXT("In"));
+	for (FPackageGraph::FEdge Edge : Container.Connections.In)
+	{
+		Writer << Edge;
+	}
+	Writer.EndArray();
+
+	Writer.BeginArray(UTF8TEXT("Out"));
+	for (FPackageGraph::FEdge Edge : Container.Connections.Out)
+	{
+		Writer << Edge;
+	}
+	Writer.EndArray();
+
+	Writer.EndObject();
+	return Writer;
+}
+
+////////////////////////////////////////////////////////////////////////////////
+static bool LoadFromCompactBinary(FCbFieldView Field, FPackageGraph::FMissingEdgeData& MissingEdge)
+{
+	if (FCbObjectView Obj = Field.AsObjectView())
+	{
+		MissingEdge.From	= FPackageGraph::FNode { .Index = Obj["From"].AsInt32() };
+		MissingEdge.To		= FPackageId::FromValue(Obj["To"].AsUInt64());
+		MissingEdge.bSoft	= Obj["Soft"].AsBool();
+
+		return true;
+	}
+
+	return false; 
+}
+
+////////////////////////////////////////////////////////////////////////////////
+static FCbWriter& operator<<(FCbWriter& Writer, const FPackageGraph::FMissingEdgeData& MissingEdge)
+{
+	Writer.BeginObject();
+	Writer << UTF8TEXT("From") << MissingEdge.From;
+	Writer << UTF8TEXT("To") << MissingEdge.To.Value();
+	Writer << UTF8TEXT("Soft") << MissingEdge.bSoft;
+	Writer.EndObject();
+
+	return Writer;
+}
+
+////////////////////////////////////////////////////////////////////////////////
+static bool LoadFromCompactBinary(FCbFieldView Field, FPackageGraph::FRedirect& Redirect)
+{
+	if (FCbObjectView Obj = Field.AsObjectView())
+	{
+		Redirect.From	= FPackageId::FromValue(Obj["From"].AsUInt64());
+		Redirect.To		= FPackageId::FromValue(Obj["To"].AsUInt64());
+
+		return true;
+	}
+
+	return false;
+}
+
+////////////////////////////////////////////////////////////////////////////////
+static FCbWriter& operator<<(FCbWriter& Writer, const FPackageGraph::FRedirect& Redirect)
+{
+	Writer.BeginObject();
+	Writer << UTF8TEXT("From") << Redirect.From.Value();
+	Writer << UTF8TEXT("To") << Redirect.To.Value();
+	Writer.EndObject();
+
+	return Writer;
+}
+
+////////////////////////////////////////////////////////////////////////////////
+FCbWriter& operator<<(FCbWriter& Writer, const FPackageGraph& Graph)
+{
+	Writer.BeginObject();
+
+	if (Graph.ContainerNodeData.IsEmpty() == false)
+	{
+		Writer.BeginArray(UTF8TEXT("Containers"));
+		for (const FPackageGraph::FContainerData& Data : Graph.ContainerNodeData)
+		{
+			Writer << Data;
+		}
+		Writer.EndArray();
+	}
+
+	if (Graph.NodeData.IsEmpty() == false)
+	{
+		Writer.BeginArray(UTF8TEXT("Nodes"));
+		for (const FPackageGraph::FNodeData& Data : Graph.NodeData)
+		{
+			Writer << Data;
+		}
+		Writer.EndArray();
+	}
+
+	if (Graph.EdgeData.IsEmpty() == false)
+	{
+		Writer.BeginArray(UTF8TEXT("Edges"));
+		for (const FPackageGraph::FEdgeData& Edge : Graph.EdgeData)
+		{
+			Writer << Edge;
+		}
+		Writer.EndArray();
+	}
+
+	if (Graph.MissingEdgeData.IsEmpty() == false)
+	{
+		Writer.BeginArray(UTF8TEXT("MissingEdges"));
+		for (const FPackageGraph::FMissingEdgeData& Missing : Graph.MissingEdgeData)
+		{
+			Writer << Missing;
+		}
+		Writer.EndArray();
+	}
+
+	if (Graph.Redirects.IsEmpty() == false)
+	{
+		Writer.BeginArray(UTF8TEXT("Redirects"));
+		for (const FPackageGraph::FRedirect& Redirect : Graph.Redirects)
+		{
+			Writer << Redirect;
+		}
+		Writer.EndArray();
+	}
+
+	Writer.EndObject();
+
+	return Writer;
+}
+
+////////////////////////////////////////////////////////////////////////////////
+bool LoadFromCompactBinary(FCbFieldView Field, FPackageGraph& Graph)
+{
+	FCbObjectView Obj = Field.AsObjectView();
+
+	if (!Obj)
+	{
+		return false;
+	}
+
+	if (FCbArrayView ArrayView = Obj["Containers"].AsArrayView())
+	{
+		Graph.ContainerNodeData.Reserve(int32(ArrayView.Num()));
+		for (FCbFieldView ArrayField : ArrayView)
+		{
+			if (!LoadFromCompactBinary(ArrayField, Graph.ContainerNodeData.AddDefaulted_GetRef()))
+			{
+				return false;
+			}
+		}
+	}
+
+	if (FCbArrayView ArrayView = Obj["Nodes"].AsArrayView())
+	{
+		Graph.NodeData.Reserve(int32(ArrayView.Num()));
+		for (FCbFieldView ArrayField : ArrayView)
+		{
+			const int32					Index = Graph.NodeData.Num();
+			FPackageGraph::FNode		Node { .Index = Index  };
+			FPackageGraph::FNodeData&	Data = Graph.NodeData.AddDefaulted_GetRef();
+
+			if (LoadFromCompactBinary(ArrayField, Data))
+			{
+				Graph.PackageIdToNode.Add(Data.PackageId, Node);
+				Graph.ContainerNodeData[Data.ContainerNode.Index].PackageNodes.Add(Node);
+			}
+			else
+			{
+				return false;
+			}
+		}
+	}
+
+	if (FCbArrayView ArrayView = Obj["Edges"].AsArrayView())
+	{
+		Graph.EdgeData.Reserve(int32(ArrayView.Num()));
+		for (FCbFieldView ArrayField : ArrayView)
+		{
+			if (!LoadFromCompactBinary(ArrayField, Graph.EdgeData.AddDefaulted_GetRef()))
+			{
+				return false;
+			}
+		}
+	}
+
+	if (FCbArrayView ArrayView = Obj["MissingEdges"].AsArrayView())
+	{
+		Graph.MissingEdgeData.Reserve(int32(ArrayView.Num()));
+		for (FCbFieldView ArrayField : ArrayView)
+		{
+			if (!LoadFromCompactBinary(ArrayField, Graph.MissingEdgeData.AddDefaulted_GetRef()))
+			{
+				return false;
+			}
+		}
+	}
+
+	if (FCbArrayView ArrayView = Obj["Redirects"].AsArrayView())
+	{
+		Graph.Redirects.Reserve(int32(ArrayView.Num()));
+		for (FCbFieldView ArrayField : ArrayView)
+		{
+			FPackageGraph::FRedirect Redirect;
+			if (LoadFromCompactBinary(ArrayField, Redirect))
+			{
+				Graph.Redirects.Add(Redirect);
+			}
+			else
+			{
+				return false;
+			}
+		}
+	}
+
+	return true;
+}
+
+////////////////////////////////////////////////////////////////////////////////
+FIoStatus LoadFromCompactBinary(const FString& Filename, FPackageGraph& OutGraph)
+{
+	IFileManager& Ifm = IFileManager::Get();
+
+	TUniquePtr<FArchive> Ar(Ifm.CreateFileReader(*Filename));
+
+	if (Ar.IsValid() == false)
+	{
+		return FIoStatusBuilder(EIoErrorCode::FileNotOpen) << TEXT("Failed to open '") << Filename << TEXT("'");
+	}
+
+	if (FCbField GraphCb = LoadCompactBinary(*Ar))
+	{
+		if (LoadFromCompactBinary(GraphCb, OutGraph))
+		{
+			return FIoStatus::Ok;
+		}
+	}
+
+	return FIoStatusBuilder(EIoErrorCode::ReadError) << TEXT("Failed to read combact binary from '") << Filename << TEXT("'");
+}
+
+////////////////////////////////////////////////////////////////////////////////
+namespace GraphTool
+{
+
+////////////////////////////////////////////////////////////////////////////////
+static TArray<FString> GlobContainers(const FContext& Context)
+{
+	IFileManager& Ifm = IFileManager::Get();
+
+	const FString GlobPattern = FString(Context.Get<FStringView>(TEXT("ContainerGlob")));
+	TArray<FString> Ret;
+
+	if (Ifm.FileExists(*GlobPattern))
+	{
+		Ret.Add(GlobPattern);
+	}
+	else if (Ifm.DirectoryExists(*GlobPattern))
+	{
+		FString Directory = GlobPattern;
+		FPaths::NormalizeDirectoryName(Directory);
+
+		TArray<FString> FoundContainerFiles;
+		Ifm.FindFiles(FoundContainerFiles, *(Directory / TEXT("*.utoc")), true, false);
+
+		for (const FString& Filename : FoundContainerFiles)
+		{
+			Ret.Emplace(Directory / Filename);
+		}
+	}
+	else
+	{
+		FString Directory = FPaths::GetPath(GlobPattern);
+		FPaths::NormalizeDirectoryName(Directory);
+
+		TArray<FString> FoundContainerFiles;
+		Ifm.FindFiles(FoundContainerFiles, *GlobPattern, true, false);
+
+		for (const FString& Filename : FoundContainerFiles)
+		{
+			Ret.Emplace(Directory / Filename);
+		}
+	}
+
+	return Ret;
+}
+
+} // UE::IoStore::Tool::GraphTool
+
+////////////////////////////////////////////////////////////////////////////////
+static int32 BuildGraphCommandEntry(const FContext& Context)
+{
+	GPrintLogCategory = GPrintLogVerbosity = false;
+	GPrintLogTimes = ELogTimes::None;
+
+	TArray<FString> Containers	= GraphTool::GlobContainers(Context);
+	FKeyChain		KeyChain	= Common::LoadCryptoKeys(Context);
+	FString			OutFilename	= FString(Context.Get<FStringView>(TEXT("-Out")));
+
+	if (Containers.IsEmpty())
+	{
+		UE_LOG(LogIoStoreOnDemand, Error, TEXT("Failed to find any container(s)."));
+		return -1;
+	}
+
+	if (OutFilename.IsEmpty())
+	{
+		UE_LOG(LogIoStoreOnDemand, Error, TEXT("No filename specified"));
+		return -1;
+	}
+
+	const bool bSkipSoftRefs = Context.Get<bool>(TEXT("-SkipSoftRefs"));
+	FPackageGraph Graph;
+	FPackageGraph::Build(Containers, KeyChain, !bSkipSoftRefs, Graph);
+	FIoStatus Status = FPackageGraph::Save(Graph, OutFilename);
+
+	if (Status.IsOk() == false)
+	{
+		UE_LOG(LogIoStoreOnDemand, Error, TEXT("%s"), *Status.ToString());
+		return -1;
+	}
+
+	return 0; 
+}
+
+////////////////////////////////////////////////////////////////////////////////
+static int32 QueryGraphCommandEntry(const FContext& Context)
+{
+	GPrintLogCategory = GPrintLogVerbosity = false;
+	GPrintLogTimes = ELogTimes::None;
+
+	FString			GraphFilename = FString(Context.Get<FStringView>(TEXT("Graph")));
+	FPackageGraph	Graph;
+
+	bool		bTraverse	= false;
+	bool		bPrintDeps	= false;
+	FPackageId	PackageId	= FPackageId();
+
+	if (FStringView HexView = Context.Get<FStringView>(TEXT("-Traverse")); HexView.IsEmpty() == false)
+	{
+		uint64 Value	= FParse::HexNumber64(FString(HexView));
+		PackageId		= FPackageId::FromValue(Value);
+		bTraverse		= true;
+	}
+
+	if (FStringView HexView = Context.Get<FStringView>(TEXT("-Deps")); HexView.IsEmpty() == false)
+	{
+		uint64 Value	= FParse::HexNumber64(FString(HexView));
+		PackageId		= FPackageId::FromValue(Value);
+		bPrintDeps		= true;
+	}
+
+	if (PackageId.IsValid() == false)
+	{
+		UE_LOG(LogIoStoreOnDemand, Error, TEXT("Invalid package ID 0x%s"), *LexToString(PackageId));
+		return -1;
+	}
+
+	UE_LOG(LogIoStoreOnDemand, Display, TEXT("Loading graph '%s'"), *GraphFilename);
+	if (FIoStatus Status = LoadFromCompactBinary(GraphFilename, Graph); !Status.IsOk()) 
+	{
+		UE_LOG(LogIoStoreOnDemand, Error, TEXT("%s"), *Status.ToString());
+		return -1;
+	}
+
+	const FPackageGraph::FNode Node = Graph.GetNode(PackageId);
+	if (Node.IsValid() == false)
+	{
+		UE_LOG(LogIoStoreOnDemand, Error, TEXT("Failed to find graph node for 0x%s"), *LexToString(PackageId));
+		return -1;
+	}
+
+	if (bTraverse)
+	{
+		struct FContainerStats
+		{
+			uint64 PackageCount = 0;
+			uint64 TotalSize	= 0;
+		};
+
+		struct FVisitor final
+			: public FPackageGraph::FVisitor
+		{
+			virtual bool Visit(const FPackageGraph& Graph, FPackageGraph::FNode Node, const FPackageGraph::FNodeData& Data) override
+			{
+				bool bAlreadyInSet = false;
+				Visited.Add(Node, &bAlreadyInSet);
+				check(!bAlreadyInSet);
+				TraverseOrder.Add(Node);
+
+				FContainerStats& Stats = ContainerStats.FindOrAdd(Data.ContainerNode);
+				Stats.TotalSize += Data.Size;
+				Stats.PackageCount++;
+				TotalSize += Data.Size;
+
+				return true;
+			}
+
+			TSet<FPackageGraph::FNode>								Visited;
+			TArray<FPackageGraph::FNode>							TraverseOrder;
+			TMap<FPackageGraph::FContainerNode, FContainerStats>	ContainerStats;
+			uint64 TotalSize = 0;
+		}; 
+
+		const bool bSkipSoftRefs	= Context.Get<bool>(TEXT("-SkipSoftRefs"));
+		const bool bIncoming		= Context.Get<bool>(TEXT("-Incoming"));
+
+		UE_LOG(LogIoStoreOnDemand, Display, TEXT("Traversing %s graph node(s) %s for package ID 0x%s"),
+			bIncoming ? TEXT("incoming") : TEXT("outgoing"),
+			bSkipSoftRefs ? TEXT("skipping soft references") : TEXT("including soft references"),
+			*LexToString(PackageId));
+		UE_LOG(LogIoStoreOnDemand, Display, TEXT(""));
+		UE_LOG(LogIoStoreOnDemand, Display, TEXT("%4s %-20s %-40s %-20s %s"),
+			TEXT("#"), TEXT("Container ID"), TEXT("Container"), TEXT("Package ID"), TEXT("Filename"));
+		UE_LOG(LogIoStoreOnDemand, Display, TEXT("----------------------------------------------------------------------------------------------------"));
+
+		FVisitor Visitor;
+		Graph.Traverse(Node, bIncoming, bSkipSoftRefs, Visitor);
+
+		for (int32 Idx = 1; FPackageGraph::FNode Traversed : Visitor.TraverseOrder)
+		{
+			const FPackageGraph::FNodeData&			Data = Graph.GetNodeData(Traversed);
+			const FPackageGraph::FContainerData&	ContainerData = Graph.GetContainerData(Data.ContainerNode);
+
+			UE_LOG(LogIoStoreOnDemand, Display, TEXT("%3d: 0x%-18s %-40s %s%-18s '%s' (%.2fKiB)"),
+				Idx++, *LexToString(ContainerData.Id), *ContainerData.Name, TEXT("0x"), *LexToString(Data.PackageId), *Data.Filename, float(Data.Size) / 1024.0f);
+		}
+
+		UE_LOG(LogIoStoreOnDemand, Display, TEXT("----------------------------------------------------------------------------------------------------"));
+		UE_LOG(LogIoStoreOnDemand, Display, TEXT("Traversed %d package(s) of total %.2f KiB in %d container(s) for package ID 0x%s"),
+				Visitor.TraverseOrder.Num(), float(Visitor.TotalSize) / 1024.0f, Visitor.ContainerStats.Num(), *LexToString(PackageId));
+
+		for (const TPair<FPackageGraph::FContainerNode, FContainerStats>& Kv : Visitor.ContainerStats)
+		{
+			const FPackageGraph::FContainerData& ContainerData = Graph.GetContainerData(Kv.Key);
+			UE_LOG(LogIoStoreOnDemand, Display, TEXT("\t* %-40s : %.2f KiB (%llu)"), *ContainerData.Name, float(Kv.Value.TotalSize) / 1024.0f, Kv.Value.PackageCount);
+		}
+	}
+	else if (bPrintDeps)
+	{
+		const bool								bIncoming		= Context.Get<bool>(TEXT("-Incoming"));
+		const FPackageGraph::FNodeData&			NodeData		= Graph.GetNodeData(Node);
+		const FPackageGraph::FContainerData&	ContainerData	= Graph.GetContainerData(NodeData.ContainerNode);
+
+		UE_LOG(LogIoStoreOnDemand, Display, TEXT(""));
+		UE_LOG(LogIoStoreOnDemand, Display, TEXT("Dependencies: 0x%s"), *LexToString(NodeData.PackageId));
+		UE_LOG(LogIoStoreOnDemand, Display, TEXT("--------------------------------------------------------------------------------"));
+		UE_LOG(LogIoStoreOnDemand, Display, TEXT("%-16s: %s"), TEXT("Container"), *ContainerData.Name);
+		UE_LOG(LogIoStoreOnDemand, Display, TEXT("%-16s: '%s'"), TEXT("Filename"), *LexToString(NodeData.Filename));
+		UE_LOG(LogIoStoreOnDemand, Display, TEXT("%-16s: %.2fKiB"), TEXT("Size"), float(NodeData.Size) / 1024.0f);
+		UE_LOG(LogIoStoreOnDemand, Display, TEXT(""));
+
+		UE_LOG(LogIoStoreOnDemand, Display, TEXT("Outgoing:"));
+		UE_LOG(LogIoStoreOnDemand, Display, TEXT("---------"));
+		for (int32 Idx = 1; FPackageGraph::FEdge Edge : NodeData.Connections.Out)
+		{
+			const FPackageGraph::FEdgeData&		EdgeData		= Graph.GetEdgeData(Edge);
+			const FPackageGraph::FNodeData&		ToData			= Graph.GetNodeData(EdgeData.To);
+			const FPackageGraph::FContainerData ToContainerData = Graph.GetContainerData(ToData.ContainerNode);
+			UE_LOG(LogIoStoreOnDemand, Display, TEXT("%d: %-40s %s%-18s '%s' (%.2f KiB)"),
+				Idx++, *ToContainerData.Name, TEXT("0x"), *LexToString(ToData.PackageId), *ToData.Filename, float(ToData.Size) / 1024.0f);
+			ensure(EdgeData.From == Node);
+		}
+
+		if (bIncoming)
+		{
+			UE_LOG(LogIoStoreOnDemand, Display, TEXT(""));
+			UE_LOG(LogIoStoreOnDemand, Display, TEXT("Incoming:"));
+			UE_LOG(LogIoStoreOnDemand, Display, TEXT("---------"));
+			for (int32 Idx = 1; FPackageGraph::FEdge Edge : NodeData.Connections.In)
+			{
+				const FPackageGraph::FEdgeData&		EdgeData			= Graph.GetEdgeData(Edge);
+				const FPackageGraph::FNodeData&		FromData			= Graph.GetNodeData(EdgeData.From);
+				const FPackageGraph::FContainerData FromContainerData	= Graph.GetContainerData(FromData.ContainerNode);
+				UE_LOG(LogIoStoreOnDemand, Display, TEXT("%d: %-40s %s%-18s '%s' (%.2f KiB)"),
+						Idx++, *FromContainerData.Name, TEXT("0x"), *LexToString(FromData.PackageId), *FromData.Filename, float(FromData.Size) / 1024.0f);
+				ensure(EdgeData.To == Node);
+			}
+		}
+	}
+
+	return 0; 
+}
+
+////////////////////////////////////////////////////////////////////////////////
+static FCommand GraphCommand(
+	BuildGraphCommandEntry,
+	TEXT("BuildGraph"),
+	TEXT(""),
+	{
+		TArgument<FStringView>(TEXT("ContainerGlob"),	TEXT("Path globbed to discover input containers")),
+		TArgument<FStringView>(TEXT("-CryptoKeys"),		TEXT("JSON-format keyring for input containers")),
+		TArgument<bool>(TEXT("-SkipSoftRefs"),			TEXT("Whether to skip traverisng soft package reference(s).")),
+		TArgument<FStringView>(TEXT("-Out"),			TEXT("Filename to save the graph (JSON|CompactBinary)."))
+	}
+);
+
+////////////////////////////////////////////////////////////////////////////////
+static FCommand QueryGraphCommand(
+	QueryGraphCommandEntry,
+	TEXT("QueryGraph"),
+	TEXT(""),
+	{
+		TArgument<FStringView>(TEXT("Graph"),		TEXT("Filepath to the package graph stored in compact binary format.")),
+		TArgument<FStringView>(TEXT("-Traverse"),	TEXT("Traverse the graph starting from the specified package ID.")),
+		TArgument<FStringView>(TEXT("-Deps"),		TEXT("Print dependencies for the specified package ID.")),
+		TArgument<bool>(TEXT("-SkipSoftRefs"),		TEXT("Whether to skip traverisng soft package reference(s).")),
+		TArgument<bool>(TEXT("-Incoming"),			TEXT("Whether to traverse incoming package nodes."))
+	}
+);
+
+} // namespace UE::IoStore::Tool
diff --git a/Engine/Source/Developer/IoStore/OnDemandUtilities/Private/Graph.h b/Engine/Source/Developer/IoStore/OnDemandUtilities/Private/Graph.h
new file mode 100644
--- /dev/null
+++ b/Engine/Source/Developer/IoStore/OnDemandUtilities/Private/Graph.h
@@ -0,0 +1,122 @@
+// Copyright Epic Games, Inc. All Rights Reserved.
+
+#include "CoreTypes.h"
+#include "Containers/Map.h"
+#include "Containers/Set.h"
+#include "IO/IoChunkId.h"
+#include "IO/IoContainerId.h"
+#include "IO/IoStatus.h"
+#include "IO/PackageId.h"
+
+class FCbWriter;
+class FCbFieldView;
+struct FKeyChain;
+
+namespace UE::IoStore::Tool
+{
+
+////////////////////////////////////////////////////////////////////////////////
+class FPackageGraph
+{
+public:
+	struct FIndex
+	{
+		bool			IsValid() const					{ return Index != INDEX_NONE; }
+		friend uint32	GetTypeHash(FIndex Node)		{ return GetTypeHash(Node.Index); }
+		explicit		operator bool() const			{ return IsValid(); }
+		bool			operator==(FIndex Other) const	{ return Index == Other.Index; }
+
+		int32 Index = INDEX_NONE;
+	};
+
+	using FNode				= FIndex;
+	using FEdge				= FIndex;
+	using FContainerNode	= FIndex;
+	using FMissingEdge		= FIndex;
+
+	struct FConnections
+	{
+		TSet<FEdge>	In;
+		TSet<FEdge>	Out;
+	};
+
+	struct FNodeData
+	{
+		FConnections		Connections;
+		FString				Filename;
+		FPackageId			PackageId;
+		uint64				Size;
+		FContainerNode		ContainerNode;
+	};
+
+	struct FEdgeData
+	{
+		FNode	From;
+		FNode	To;
+		bool	bSoft = false;
+	};
+
+	struct FMissingEdgeData
+	{
+		FNode		From;
+		FPackageId	To;
+		bool		bSoft = false;
+	};
+
+	struct FContainerData
+	{
+		FConnections	Connections;
+		TSet<FNode>		PackageNodes; // Transient
+		FString			Name;
+		FIoContainerId	Id;
+	};
+
+	struct FRedirect
+	{
+		friend uint32	GetTypeHash(const FRedirect &Redirect)		{ return GetTypeHash(Redirect.From); }
+		bool			operator==(const FRedirect &Other) const	{ return From == Other.From; }
+		bool			operator<(const FRedirect &Other) const		{ return From < Other.From; }
+
+		FPackageId From;
+		FPackageId To;
+	};
+
+	struct FVisitor
+	{
+		virtual bool Visit(const FPackageGraph& Graph, FNode Node, const FNodeData& Data) = 0;
+	};
+
+	void						Reserve(int32 ContainerCount, int32 PackageCount);
+	FContainerNode				AddContainerNode(FString&& Name, const FIoContainerId& Id);
+	const FContainerData&		GetContainerData(FContainerNode Node) const { return ContainerNodeData[Node.Index]; }
+	FNode						AddNode(FContainerNode ContainerNode, const FPackageId& PackageId, const FString* Filename, uint64 ChunkSize);
+	FNode						GetNode(const FPackageId& PackageId) const { return PackageIdToNode.FindRef(PackageId); }
+	const FNodeData&			GetNodeData(FNode Node) const { return NodeData[Node.Index]; }
+	int32						GetNodeCount() const { return NodeData.Num(); }
+	void						AddEdge(FNode From, FNode To, bool bSoft);
+	const FEdgeData&			GetEdgeData(FEdge Edge) const { return EdgeData[Edge.Index]; }
+	int32						GetEdgeCount() const { return EdgeData.Num(); }
+	void						AddMissingEdge(FNode From, const FPackageId To, bool bSoft);
+	void						AddRedirect(const FPackageId& From, const FPackageId& To);
+
+	void						Traverse(FNode Node, bool bIncoming, bool bSkipSoftRefs, FVisitor& Visitor);
+
+	friend FCbWriter&			operator<<(FCbWriter& Writer, const FPackageGraph& Graph);
+	friend bool					LoadFromCompactBinary(FCbFieldView Field, FPackageGraph& Graph);
+
+	static void					Build(TConstArrayView<FString> ContainerFilenames, const FKeyChain& KeyChain, bool bSkipSoftRefs, FPackageGraph& Graph);
+	static FIoStatus			Save(const FPackageGraph& Graph, const FString& Filename);
+
+private:
+	TArray<FNodeData>			NodeData;
+	TArray<FEdgeData>			EdgeData;
+	TArray<FMissingEdgeData>	MissingEdgeData;
+	TArray<FContainerData>		ContainerNodeData;
+	TSet<FRedirect>				Redirects;
+	TMap<FPackageId, FNode>		PackageIdToNode; // Transient
+};
+
+////////////////////////////////////////////////////////////////////////////////
+FIoStatus	LoadFromCompactBinary(const FString& Filename, FPackageGraph& OutGraph);
+
+} // namespace UE::IoStore::Tool
diff --git a/Engine/Source/Developer/IoStore/OnDemandUtilities/Private/ListTocs.cpp b/Engine/Source/Developer/IoStore/OnDemandUtilities/Private/ListTocs.cpp
new file mode 100644
--- /dev/null
+++ b/Engine/Source/Developer/IoStore/OnDemandUtilities/Private/ListTocs.cpp
@@ -0,0 +1,760 @@
+// Copyright Epic Games, Inc. All Rights Reserved.
+
+#include "Command.h"
+
+#include "Async/Mutex.h"
+#include "Async/ParallelFor.h"
+#include "Async/UniqueLock.h"
+#include "HAL/FileManager.h"
+#include "IO/IoContainerHeader.h"
+#include "IO/IoStoreOnDemand.h"
+#include "Misc/EnumClassFlags.h"
+#include "Misc/FileHelper.h"
+#include "Misc/Paths.h"
+#include "S3/S3Client.h"
+#include "Serialization/JsonSerializer.h"
+#include "Serialization/JsonWriter.h"
+#include "Serialization/MemoryReader.h"
+
+namespace UE::IoStore::Tool
+{
+
+using FJsonWriter = TSharedPtr<TJsonWriter<TCHAR, TPrettyJsonPrintPolicy<TCHAR>>>; 
+using FJsonWriterFactory = TJsonWriterFactory<TCHAR, TPrettyJsonPrintPolicy<TCHAR>>;
+
+////////////////////////////////////////////////////////////////////////////////
+
+/**
+ * Parameters for listing uploaded TOC file(s) from an S3 compatible endpoint.
+ *
+ * Example usage:
+ *
+ * 1) Print available TOC's from a local server to standard out.
+ * IasTool.exe ListTocs -ServiceUrl="http://10.24.101.92:9000" -Bucket=<bucketname> -BucketPrefix=<some/path/to/data> -AccessKey=<accesskey> -SecretKey=<secretkey>
+ *
+ * 2) Print available TOC's from AWS S3.
+ * IasTool.exe ListTocs -Region="us-east-1" -BucketPath=<mybucket/some/data/path/> -CredentialsFile=<path/to/credentials.txt> -BuildVersion=<version> -Json=<path/to/file.json>
+ * 
+ * 3) Serialize all TOC's matching a specific build version to JSON:
+ * IasTool.exe ListTocs -Region="us-east-1" -BucketPath=<mybucket/some/data/path/> -CredentialsFile=<path/to/credentials.txt> -BuildVersion=<version> -Json=<path/to/file.json>
+ *
+ * 4) Serialize all chunk object key(s) to JSON.
+ * IasTool.exe ListTocs -Region="us-east-1" -BucketPath=<mybucket/some/data/path/> -CredentialsFile=<path/to/credentials.txt> -BuildVersion=<version> -ChunkKeys=<path/to/file.json>
+ *
+ * 5) Fetch a TOC from a public CDN.
+ * IasTool.exe ListTocs -TocUrl=<http://some.public.endpoint.net/path/to/1a32076ca12bfc6feb982ffb064d18f28156606c.iochunktoc>
+ *
+ * Parameters: -TocEntries, -BlockSizes and -BlockHashes controls what to include when serializing TOC's to JSON.
+ *
+ * Credentials file example:
+ *
+ * [default]
+ * aws_access_key_id="<key>"
+ * aws_secret_access_key="<key>
+ * aws_session_token="<token>"
+ *
+ * Note: All values must be surrounded with "".
+ */
+
+////////////////////////////////////////////////////////////////////////////////
+
+enum class EOnDemandTocJsonOptions : uint32
+{
+	Header			= (1 << 0),
+	TocEntries		= (1 << 1),
+	BlockSizes		= (1 << 2),
+	BlockHashes		= (1 << 3),
+	ContainerHeader	= (1 << 4),
+
+	All				= Header | TocEntries | BlockSizes | BlockHashes | ContainerHeader
+};
+ENUM_CLASS_FLAGS(EOnDemandTocJsonOptions);
+
+EOnDemandTocJsonOptions LexFromContext(const FContext& Context)
+{
+	if (Context.Get<bool>(TEXT("-AllFilters"), false))
+	{
+		return EOnDemandTocJsonOptions::All;
+	}
+
+	EOnDemandTocJsonOptions JsonOptions = EOnDemandTocJsonOptions::Header;
+
+	if (Context.Get<bool>(TEXT("-TocEntries"), false))
+	{
+		JsonOptions |= EOnDemandTocJsonOptions::TocEntries;
+	}
+
+	if (Context.Get<bool>(TEXT("-BlockSizes"), false))
+	{
+		JsonOptions |= EOnDemandTocJsonOptions::BlockSizes;
+	}
+
+	if (Context.Get<bool>(TEXT("-BlockSizes"), false))
+	{
+		JsonOptions |= EOnDemandTocJsonOptions::BlockHashes;
+	}
+
+	if (Context.Get<bool>(TEXT("-ContainerHeader"), false))
+	{
+		JsonOptions |= EOnDemandTocJsonOptions::ContainerHeader;
+	}
+
+	return JsonOptions;
+}
+
+////////////////////////////////////////////////////////////////////////////////
+
+struct FListTocsParams
+{
+	FString OutFile;
+	FString ServiceUrl;
+	FString Bucket;
+	FString BucketPrefix;
+	FString Region; 
+	FString AccessKey;
+	FString SecretKey;
+	FString SessionToken;
+	FString CredentialsFile;
+	FString CredentialsFileKeyName;
+	FString TocUrl;
+	FString TocPath;
+	FString TocKey;
+	FString BuildVersion;
+	FString TargetPlatform;
+	FString ChunkKeys;
+	EOnDemandTocJsonOptions JsonOptions = EOnDemandTocJsonOptions::Header;
+
+	FIoStatus Validate() const;
+};
+
+FIoStatus FListTocsParams::Validate() const
+{
+	if (!TocUrl.IsEmpty())
+	{
+		return FIoStatus::Ok;
+	}
+	if (!TocPath.IsEmpty())
+	{
+		return FIoStatus::Ok;
+	}
+	if (Bucket.IsEmpty())
+	{
+		return FIoStatus(EIoErrorCode::InvalidParameter, TEXT("Invalid bucket name"));
+	}
+	if (!AccessKey.IsEmpty() && SecretKey.IsEmpty())
+	{
+		return FIoStatus(EIoErrorCode::InvalidParameter, TEXT("Invalid secret key"));
+	}
+	else if (AccessKey.IsEmpty() && !SecretKey.IsEmpty())
+	{
+		return FIoStatus(EIoErrorCode::InvalidParameter, TEXT("Invalid access key"));
+	}
+
+	if (!CredentialsFile.IsEmpty() && CredentialsFileKeyName.IsEmpty())
+	{
+		return FIoStatus(EIoErrorCode::InvalidParameter, TEXT("Invalid credential file key name"));
+	}
+
+	if (TocUrl.IsEmpty() && ServiceUrl.IsEmpty() && Region.IsEmpty())
+	{
+		return FIoStatus(EIoErrorCode::InvalidParameter, TEXT("Service URL or AWS region needs to be specified"));
+	}
+
+	return FIoStatus::Ok;
+}
+
+////////////////////////////////////////////////////////////////////////////////
+static FListTocsParams BuildListTocsParams(const FContext& Context)
+{
+	FListTocsParams Ret;
+
+	if (auto Value = Context.Get<FStringView>(TEXT("-BucketPath")); !Value.IsEmpty())
+	{
+		int32 Idx = INDEX_NONE;
+		if (Value.FindChar(TCHAR('/'), Idx))
+		{
+			Ret.Bucket = Value.Left(Idx);
+			FStringView BucketPrefix = Value.RightChop(Idx + 1);
+			if (BucketPrefix.EndsWith(TEXT("/")))
+			{
+				BucketPrefix.LeftChopInline(1);
+			}
+			Ret.BucketPrefix = BucketPrefix;
+		}
+		else
+		{
+			Ret.Bucket = Value;
+		}
+	}
+
+	Ret.OutFile					= Context.Get<FStringView>(TEXT("-Json"),					Ret.OutFile);
+	Ret.Bucket					= Context.Get<FStringView>(TEXT("-Bucket"),					Ret.Bucket);
+	Ret.BucketPrefix			= Context.Get<FStringView>(TEXT("-BucketPrefix"),			Ret.BucketPrefix);
+	Ret.ServiceUrl				= Context.Get<FStringView>(TEXT("-ServiceUrl"),				Ret.ServiceUrl);
+	Ret.Region					= Context.Get<FStringView>(TEXT("-Region"),					Ret.Region);
+	Ret.AccessKey				= Context.Get<FStringView>(TEXT("-AccessKey"),				Ret.AccessKey);
+	Ret.SecretKey				= Context.Get<FStringView>(TEXT("-SecretKey"),				Ret.SecretKey);
+	Ret.SessionToken			= Context.Get<FStringView>(TEXT("-SessionToken"),			Ret.SessionToken);
+	Ret.CredentialsFile			= Context.Get<FStringView>(TEXT("-CredentialsFile"),		Ret.CredentialsFile);
+	Ret.CredentialsFileKeyName	= Context.Get<FStringView>(TEXT("-CredentialsFileKeyName"),	Ret.CredentialsFileKeyName);
+	Ret.TocUrl					= Context.Get<FStringView>(TEXT("-TocUrl"),					Ret.TocUrl);
+	Ret.TocPath					= Context.Get<FStringView>(TEXT("-TocPath"),				Ret.TocPath);
+	Ret.TocKey					= Context.Get<FStringView>(TEXT("-TocKey"),					Ret.TocKey);
+	Ret.BuildVersion			= Context.Get<FStringView>(TEXT("-BuildVersion"),			Ret.BuildVersion);
+	Ret.TargetPlatform			= Context.Get<FStringView>(TEXT("-TargetPlatform"),			Ret.TargetPlatform);
+	Ret.ChunkKeys				= Context.Get<FStringView>(TEXT("-ChunkKeys"),				Ret.ChunkKeys);
+
+	Ret.JsonOptions = LexFromContext(Context);
+
+	if (!Ret.TocUrl.IsEmpty())
+	{
+		FStringView View = Ret.TocUrl;
+		if (View.StartsWith(TEXTVIEW("http://")) && View.EndsWith(TEXTVIEW(".iochunktoc")))
+		{
+			int32 Delim = INDEX_NONE;
+			if (View.RightChop(7).FindChar(TEXT('/'), Delim))
+			{
+				Ret.ServiceUrl = FString(View.Left(7 +  Delim));
+				Ret.TocKey = View.RightChop(Ret.ServiceUrl.Len() + 1);
+			}
+		}
+	}
+
+	if (FIoStatus Validation = Ret.Validate(); !Validation.IsOk())
+	{
+		FString Reason = Validation.ToString();
+		Context.Abort(*Reason);
+	}
+
+	return Ret;
+}
+
+////////////////////////////////////////////////////////////////////////////////
+static void GetChunkObjectKeys(const FOnDemandToc& Toc, FString Prefix, TArray<FString>& OutKeys)
+{
+	if (Prefix.EndsWith(TEXT("/")))
+	{
+		Prefix = Prefix.Left(Prefix.Len() - 1);
+	}
+
+	TStringBuilder<256> Sb;
+	for (const FOnDemandTocContainerEntry& Container : Toc.Containers)
+	{
+		for (const FOnDemandTocEntry& Entry : Container.Entries)
+		{
+			Sb.Reset();
+			const FString Hash = LexToString(Entry.Hash);
+			Sb << Prefix << TEXT("/") << Toc.Header.ChunksDirectory << TEXT("/") << Hash.Left(2) << TEXT("/") << Hash << TEXT(".iochunk");
+			OutKeys.Add(FString(Sb.ToString()).ToLower());
+		}
+	}
+}
+
+
+////////////////////////////////////////////////////////////////////////////////
+static void ToJson(FJsonWriter JsonWriter, const FOnDemandToc& Toc, EOnDemandTocJsonOptions Options)
+{
+	JsonWriter->WriteObjectStart(TEXT("Header"));
+	{
+		JsonWriter->WriteValue(TEXT("Magic"), Toc.Header.Magic);
+		JsonWriter->WriteValue(TEXT("Version"), uint64(Toc.Header.Version));
+		JsonWriter->WriteValue(TEXT("BlockSize"), uint64(Toc.Header.BlockSize));
+		JsonWriter->WriteValue(TEXT("CompressionFormat"), Toc.Header.CompressionFormat);
+		JsonWriter->WriteValue(TEXT("ChunksDirectory"), Toc.Header.ChunksDirectory);
+		JsonWriter->WriteValue(TEXT("HostGroupName"), Toc.Header.HostGroupName);
+	}
+	JsonWriter->WriteObjectEnd();
+	JsonWriter->WriteObjectStart(TEXT("Meta"));
+	{
+		const FDateTime Dt = FDateTime::FromUnixTimestamp(Toc.Meta.EpochTimestamp);
+		JsonWriter->WriteValue(TEXT("DateTime"), Dt.ToString()), 
+		JsonWriter->WriteValue(TEXT("BuildVersion"), Toc.Meta.BuildVersion);
+		JsonWriter->WriteValue(TEXT("TargetPlatform"), Toc.Meta.TargetPlatform);
+	}
+	JsonWriter->WriteObjectEnd();
+	JsonWriter->WriteArrayStart(TEXT("Containers"));
+	{
+		for (const FOnDemandTocContainerEntry& Container : Toc.Containers)
+		{
+			JsonWriter->WriteObjectStart();
+			JsonWriter->WriteValue(TEXT("Name"), Container.ContainerName);
+			JsonWriter->WriteValue(TEXT("EncryptionKeyGuid"), Container.EncryptionKeyGuid);
+			JsonWriter->WriteValue(TEXT("UTocHash"), LexToString(Container.UTocHash));
+
+			if (EnumHasAnyFlags(Options, EOnDemandTocJsonOptions::ContainerHeader))
+			{
+				JsonWriter->WriteObjectStart(TEXT("Header"));
+				{
+					FIoContainerHeader ContainerHeader;
+					FMemoryReaderView Ar(Container.Header);
+					Ar << ContainerHeader;
+
+					JsonWriter->WriteValue(TEXT("ContainerId"), LexToString(ContainerHeader.ContainerId));
+					
+					JsonWriter->WriteArrayStart("PackageIds");
+					for (const FPackageId& PackageId : ContainerHeader.PackageIds)
+					{
+						JsonWriter->WriteValue(LexToString(PackageId));
+					}
+					JsonWriter->WriteArrayEnd();
+
+					JsonWriter->WriteArrayStart("OptionalSegmentPackageIds");
+					for (const FPackageId& PackageId : ContainerHeader.OptionalSegmentPackageIds)
+					{
+						JsonWriter->WriteValue(LexToString(PackageId));
+					}
+					JsonWriter->WriteArrayEnd();
+				}
+				JsonWriter->WriteObjectEnd();
+			}
+
+			if (EnumHasAnyFlags(Options, EOnDemandTocJsonOptions::TocEntries))
+			{
+				JsonWriter->WriteArrayStart("Entries");
+				for (const FOnDemandTocEntry& Entry : Container.Entries)
+				{
+					JsonWriter->WriteObjectStart();
+					JsonWriter->WriteValue(TEXT("Hash"), LexToString(Entry.Hash));
+					JsonWriter->WriteValue(TEXT("ChunkId"), LexToString(Entry.ChunkId));
+					JsonWriter->WriteValue(TEXT("RawSize"), Entry.RawSize);
+					JsonWriter->WriteValue(TEXT("EncodedSize"), Entry.EncodedSize);
+					JsonWriter->WriteValue(TEXT("BlockOffset"), int32(Entry.BlockOffset));
+					JsonWriter->WriteValue(TEXT("BlockCount"), int32(Entry.BlockCount));
+					JsonWriter->WriteObjectEnd();
+				}
+				JsonWriter->WriteArrayEnd();
+			}
+
+			if (EnumHasAnyFlags(Options, EOnDemandTocJsonOptions::BlockSizes))
+			{
+				JsonWriter->WriteArrayStart("Blocks");
+				{
+					for (uint32 BlockSize : Container.BlockSizes)
+					{
+						JsonWriter->WriteValue(int32(BlockSize));
+					}
+				}
+				JsonWriter->WriteArrayEnd();
+			}
+
+			if (EnumHasAnyFlags(Options, EOnDemandTocJsonOptions::BlockHashes))
+			{
+				JsonWriter->WriteArrayStart("BlockHashes");
+				{
+					for (const FIoBlockHash& BlockHash : Container.BlockHashes)
+					{
+						JsonWriter->WriteValue(int32(BlockHash));
+					}
+				}
+				JsonWriter->WriteArrayEnd();
+			}
+			JsonWriter->WriteObjectEnd();
+		}
+	}
+	JsonWriter->WriteArrayEnd();
+}
+
+////////////////////////////////////////////////////////////////////////////////
+FIoStatus ListTocs(const FListTocsParams& Params)
+{
+	struct FChunkStats
+	{
+		TSet<FIoHash> Chunks;
+		uint64 TotalChunkSize = 0;
+
+		void Add(const FIoHash& Hash, uint64 ChunkSize)
+		{
+			bool bExist = false;
+			Chunks.Add(Hash, &bExist);
+			if (!bExist)
+			{
+				TotalChunkSize += ChunkSize;
+			}
+		}
+	};
+	TMap<FString, FChunkStats> Stats;
+
+	struct FTocDescription
+	{
+		FOnDemandToc Toc;
+		FDateTime DateTime;
+		FString Key;
+		uint64 Size = 0;
+		uint64 ChunkCount = 0;
+		uint64 TotalChunkSize = 0;
+	};
+
+	const bool bFilteredQuery = !Params.TocKey.IsEmpty() || !Params.BuildVersion.IsEmpty() || !Params.TargetPlatform.IsEmpty();
+
+	FS3ClientConfig Config;
+	Config.ServiceUrl = Params.ServiceUrl;
+	Config.Region = Params.Region;
+	
+	FS3ClientCredentials Credentials;
+	if (Params.CredentialsFile.IsEmpty() == false)
+	{
+		UE_LOG(LogIas, Display, TEXT("Loading credentials file '%s'"), *Params.CredentialsFile);
+		FS3CredentialsProfileStore CredentialsStore = FS3CredentialsProfileStore::FromFile(Params.CredentialsFile);
+		if (CredentialsStore.TryGetCredentials(Params.CredentialsFileKeyName, Credentials) == false)
+		{
+			return FIoStatus(EIoErrorCode::InvalidParameter, TEXT("Failed to find valid credentials in credentials file"));
+		}
+		else
+		{
+			UE_LOG(LogIas, Display, TEXT("Found credentials for '%s'"), *Params.CredentialsFileKeyName);
+		}
+	}
+	else
+	{
+		Credentials = FS3ClientCredentials(Params.AccessKey, Params.SecretKey, Params.SessionToken);
+	}
+
+	FS3Client Client(Config, Credentials);
+	FMutex TocMutex;
+	TArray<FTocDescription> Tocs;
+
+	if (!Params.TocPath.IsEmpty())
+	{
+		UE_LOG(LogIas, Display, TEXT("Loading TOC from file '%s'"), *Params.TocPath);
+
+		const int64 FileSize = IFileManager::Get().FileSize(*Params.TocPath);
+		if (FileSize <= 0)
+		{
+			return FIoStatus(EIoErrorCode::NotFound);
+		}
+
+		TIoStatusOr<FOnDemandToc> TocStatus = FOnDemandToc::LoadFromFile(Params.TocPath, false);
+		if (!TocStatus.IsOk())
+		{
+			return TocStatus.Status();
+		}
+
+		FOnDemandToc Toc = TocStatus.ConsumeValueOrDie();
+		if (Toc.Header.Magic == FOnDemandTocHeader::ExpectedMagic)
+		{
+			FDateTime DateTime = FDateTime::FromUnixTimestamp(Toc.Meta.EpochTimestamp);
+			Tocs.Add(FTocDescription 
+			{
+				.Toc = MoveTemp(Toc),
+				.DateTime = DateTime,
+				.Size = static_cast<uint64>(FileSize)
+			});
+		}
+		else
+		{
+			return FIoStatus(EIoErrorCode::FileOpenFailed, TEXT("Header mismatch"));
+		}
+	}
+	else if (Params.TocUrl.IsEmpty())
+	{
+		FStringView PrefixView = Params.BucketPrefix;
+		if (PrefixView.StartsWith(TEXT("/")))
+		{
+			PrefixView.RightChopInline(1);
+		}
+
+		TStringBuilder<256> Path;
+		if (!PrefixView.IsEmpty())
+		{
+			Path << PrefixView << TEXT("/");
+		}
+
+		UE_LOG(LogIas, Display, TEXT("Fetching TOC's from '%s/%s/%s'"), *Client.GetConfig().ServiceUrl, *Params.Bucket, Path.ToString());
+		FS3ListObjectResponse Response = Client.ListObjects(FS3ListObjectsRequest
+		{
+			Params.Bucket,
+			Path.ToString(),
+			TEXT('/')
+		});
+
+		if (Response.Objects.IsEmpty())
+		{
+			UE_LOG(LogIas, Display, TEXT("Not TOC's found at '%s/%s/%s' (%s)"), *Client.GetConfig().ServiceUrl, *Params.Bucket, Path.ToString(), *Response.GetErrorStatus());
+			return FIoStatus(EIoErrorCode::NotFound);
+		}
+
+		ParallelFor(Response.Objects.Num(), [&Params, &Client, &Response, &Tocs, &TocMutex](int32 Index)
+		{
+			const FS3Object& Obj = Response.Objects[Index];
+			if (Obj.Key.EndsWith(TEXT("iochunktoc")) == false)
+			{
+				return;
+			}
+
+			if (!Params.TocKey.IsEmpty() && !Params.TocKey.Equals(FPaths::GetBaseFilename(Obj.Key), ESearchCase::IgnoreCase))
+			{
+				return;
+			}
+
+			UE_LOG(LogIas, Display, TEXT("Fetching TOC '%s/%s/%s'"), *Client.GetConfig().ServiceUrl, *Params.Bucket, *Obj.Key);
+			FS3GetObjectResponse TocResponse = Client.GetObject(FS3GetObjectRequest
+			{
+				Params.Bucket,
+				Obj.Key
+			});
+
+			if (TocResponse.IsOk() == false)
+			{
+				UE_LOG(LogIas, Warning, TEXT("Failed to fetch TOC '%s/%s/%s'"), *Client.GetConfig().ServiceUrl, *Params.Bucket, *Obj.Key);
+				return;
+			}
+
+			FOnDemandToc Toc;
+			FMemoryReaderView Ar(TocResponse.GetBody().GetView());
+			Ar << Toc;
+
+			if (Ar.IsError() || Toc.Header.Magic != FOnDemandTocHeader::ExpectedMagic)
+			{
+				UE_LOG(LogIas, Warning, TEXT("Failed to serialize TOC '%s/%s/%s'. Header version/magic mimsatch"), *Client.GetConfig().ServiceUrl, *Params.Bucket, *Obj.Key);
+				return;
+			}
+
+			if (!Params.BuildVersion.IsEmpty() && !Params.BuildVersion.Equals(Toc.Meta.BuildVersion, ESearchCase::IgnoreCase))
+			{
+				return;
+			}
+
+			if (!Params.TargetPlatform.IsEmpty() && !Params.TargetPlatform.Equals(Toc.Meta.TargetPlatform, ESearchCase::IgnoreCase))
+			{
+				return;
+			}
+
+			FDateTime DateTime = FDateTime::FromUnixTimestamp(Toc.Meta.EpochTimestamp);
+			FTocDescription Desc = FTocDescription 
+			{
+				.Toc = MoveTemp(Toc),
+				.DateTime = DateTime,
+				.Key = Obj.Key,
+				.Size = Obj.Size
+			};
+
+			{
+				TUniqueLock Lock(TocMutex);
+				Tocs.Add(MoveTemp(Desc));
+			}
+		});
+	}
+	else
+	{
+		UE_LOG(LogIas, Display, TEXT("Fetching TOC '%s/%s'"), *Client.GetConfig().ServiceUrl, *Params.TocKey);
+		FS3GetObjectResponse TocResponse = Client.GetObject(FS3GetObjectRequest
+		{
+			.Key = Params.TocKey
+		});
+
+		if (TocResponse.IsOk())
+		{
+			FOnDemandToc Toc;
+			FMemoryReaderView Ar(TocResponse.GetBody().GetView());
+			Ar << Toc;
+
+			if (Toc.Header.Magic == FOnDemandTocHeader::ExpectedMagic)
+			{
+				FDateTime DateTime = FDateTime::FromUnixTimestamp(Toc.Meta.EpochTimestamp);
+				Tocs.Add(FTocDescription 
+				{
+					.Toc = MoveTemp(Toc),
+					.DateTime = DateTime,
+					.Key = Params.TocKey,
+					.Size = TocResponse.GetBody().GetSize()
+				});
+			}
+			else
+			{
+				UE_LOG(LogIas, Warning, TEXT("Failed to serialize TOC '%s/%s'. Header magic mimsatch"), *Client.GetConfig().ServiceUrl, *Params.TocKey);
+			}
+		}
+	}
+
+	if (Tocs.IsEmpty())
+	{
+		return FIoStatus(EIoErrorCode::NotFound);
+	}
+
+	for (FTocDescription& Desc : Tocs)
+	{
+		FChunkStats& BuildStats = Stats.FindOrAdd(Desc.Toc.Meta.BuildVersion);
+		FChunkStats& BucketStats = Stats.FindOrAdd(Params.Bucket);
+		for (const FOnDemandTocContainerEntry& Container : Desc.Toc.Containers)
+		{
+			for (const FOnDemandTocEntry& Entry : Container.Entries)
+			{
+				Desc.ChunkCount++;
+				Desc.TotalChunkSize += Entry.EncodedSize;
+				BucketStats.Add(Entry.Hash, Entry.EncodedSize);
+				BuildStats.Add(Entry.Hash, Entry.EncodedSize);
+			}
+		}
+	}
+
+	Tocs.Sort([](const auto& LHS, const auto& RHS) { return LHS.DateTime > RHS.DateTime; });
+
+	if (Params.OutFile.IsEmpty())
+	{
+		int32 Counter = 1;
+		TStringBuilder<256> Url;
+
+		for (const FTocDescription& Desc : Tocs)
+		{
+			Url.Reset();
+			Url << Client.GetConfig().ServiceUrl << TEXT("/") << Params.Bucket << TEXT("/") << Desc.Key;
+
+			UE_LOG(LogIas, Display, TEXT(""));
+			UE_LOG(LogIas, Display, TEXT("%d) %s"), Counter++, *Desc.Key);
+			UE_LOG(LogIas, Display, TEXT("%-20s: %s"),TEXT("Date"), *Desc.DateTime.ToString());
+			UE_LOG(LogIas, Display, TEXT("%-20s: %s"), TEXT("BuildVersion"), *Desc.Toc.Meta.BuildVersion);
+			UE_LOG(LogIas, Display, TEXT("%-20s: %s"), TEXT("TargetPlatform"), *Desc.Toc.Meta.TargetPlatform);
+			UE_LOG(LogIas, Display, TEXT("%-20s: %.2lf KiB"), TEXT("TocSize"), double(Desc.Size) / 1024.0);
+			UE_LOG(LogIas, Display, TEXT("%-20s: %llu"), TEXT("ChunkCount"), Desc.ChunkCount);
+			UE_LOG(LogIas, Display, TEXT("%-20s: %.2lf MiB"), TEXT("TotalChunkSize"), double(Desc.TotalChunkSize) / 1024.0 / 1024.0);
+			if (Params.TocPath.IsEmpty())
+			{
+				UE_LOG(LogIas, Display, TEXT("%-20s: %s"), TEXT("Url"), Url.ToString());
+			}
+			else
+			{
+				UE_LOG(LogIas, Display, TEXT("%-20s: %s"), TEXT("FilePath"), *Params.TocPath);
+			}
+		}
+		UE_LOG(LogIas, Display, TEXT(""));
+		
+		if (bFilteredQuery == false)
+		{
+			TArray<FString> Keys;
+			Stats.GetKeys(Keys);
+			Keys.Sort();
+
+			UE_LOG(LogIas, Display, TEXT("%-80s %15s %15s"), TEXT("BuildVersion"), TEXT("Chunk(s)"), TEXT("MiB"));
+			UE_LOG(LogIas, Display, TEXT("-------------------------------------------------------------------------------------------------------------------"));
+			for (const FString& Key : Keys)
+			{
+				if (Key != Params.Bucket)
+				{
+					const FChunkStats& ChunkStats = Stats.FindChecked(Key);
+					UE_LOG(LogIas, Display, TEXT("%-80s %15d %15.2f"), *Key, ChunkStats.Chunks.Num(), double(ChunkStats.TotalChunkSize) / 1024.0 / 1024.0);
+				}
+			}
+			UE_LOG(LogIas, Display, TEXT("-------------------------------------------------------------------------------------------------------------------"));
+			FChunkStats& BucketStats = Stats.FindOrAdd(Params.Bucket);
+			UE_LOG(LogIas, Display, TEXT("%-80s %15d %15.2f"), *Params.Bucket, BucketStats.Chunks.Num(), double(BucketStats.TotalChunkSize) / 1024.0 / 1024.0);
+			UE_LOG(LogIas, Display, TEXT(""));
+		}
+	}
+	else
+	{
+		FString Json;
+		FJsonWriter JsonWriter = FJsonWriterFactory::Create(&Json);
+		JsonWriter->WriteObjectStart();
+		JsonWriter->WriteValue(TEXT("ServiceUrl"), Client.GetConfig().ServiceUrl);
+		JsonWriter->WriteValue(TEXT("Bucket"), Params.Bucket);
+		JsonWriter->WriteValue(TEXT("BucketPrefix"), Params.BucketPrefix);
+
+		JsonWriter->WriteArrayStart(TEXT("Tocs"));
+		for (const FTocDescription& Desc : Tocs)
+		{
+			JsonWriter->WriteObjectStart();
+			JsonWriter->WriteValue(TEXT("Key"), Desc.Key);
+			JsonWriter->WriteValue(TEXT("ChunkCount"), Desc.ChunkCount);
+			JsonWriter->WriteValue(TEXT("TotalChunkSize"), Desc.TotalChunkSize);
+
+			ToJson(JsonWriter, Desc.Toc, Params.JsonOptions);
+			JsonWriter->WriteObjectEnd();
+		}
+		JsonWriter->WriteArrayEnd();
+
+		JsonWriter->WriteObjectEnd();
+		JsonWriter->Close();
+
+		UE_LOG(LogIas, Display, TEXT("Saving file '%s'"), *Params.OutFile);
+		if (!FFileHelper::SaveStringToFile(Json, *Params.OutFile))
+		{
+			return FIoStatus(EIoErrorCode::WriteError, TEXTVIEW("Failed writing JSON file")); 
+		}
+	}
+
+	if (Params.ChunkKeys.IsEmpty() == false)
+	{
+		FString Json;
+		FJsonWriter JsonWriter = FJsonWriterFactory::Create(&Json);
+		
+		JsonWriter->WriteObjectStart();
+		JsonWriter->WriteValue(TEXT("ServiceUrl"), Client.GetConfig().ServiceUrl);
+		JsonWriter->WriteValue(TEXT("Bucket"), Params.Bucket);
+		JsonWriter->WriteValue(TEXT("BucketPrefix"), Params.BucketPrefix);
+
+		JsonWriter->WriteArrayStart(TEXT("Tocs"));
+		for (const FTocDescription& Desc : Tocs)
+		{
+			TArray<FString> ObjKeys;
+			GetChunkObjectKeys(Desc.Toc, Params.BucketPrefix, ObjKeys);
+
+			JsonWriter->WriteObjectStart();
+			JsonWriter->WriteValue(TEXT("Key"), Desc.Key);
+			JsonWriter->WriteValue(TEXT("BuildVersion"), Desc.Toc.Meta.BuildVersion);
+			JsonWriter->WriteValue(TEXT("TargetPlatform"), Desc.Toc.Meta.TargetPlatform);
+			JsonWriter->WriteArrayStart(TEXT("ChunkKeys"));
+			for (const FString& Key : ObjKeys)
+			{
+				JsonWriter->WriteValue(Key);
+			}
+			JsonWriter->WriteArrayEnd();
+			JsonWriter->WriteObjectEnd();
+		}
+		JsonWriter->WriteArrayEnd();
+
+		JsonWriter->WriteObjectEnd();
+		JsonWriter->Close();
+
+		UE_LOG(LogIas, Display, TEXT("Saving chunk key(s) '%s'"), *Params.ChunkKeys);
+		if (!FFileHelper::SaveStringToFile(Json, *Params.ChunkKeys))
+		{
+			return FIoStatus(EIoErrorCode::WriteError, TEXTVIEW("Failed writing JSON file")); 
+		}
+	}
+
+	return FIoStatus::Ok;
+}
+
+////////////////////////////////////////////////////////////////////////////////
+static int32 ListTocsCommandEntry(const Tool::FContext& Context)
+{
+	FListTocsParams Params = BuildListTocsParams(Context);
+
+	FIoStatus Status = ListTocs(Params);
+	if (!Status.IsOk())
+	{
+		FString Reason = Status.ToString();
+		Context.Abort(*Reason);
+	}
+
+	return 0;
+}
+
+////////////////////////////////////////////////////////////////////////////////
+extern FArgumentSet S3Arguments;
+
+static FCommand ListTocsCommand(
+	ListTocsCommandEntry,
+	TEXT("ListTocs"),
+	TEXT("Enumerates and output information for available TOCs"),
+	{
+		TArgument<FStringView>(TEXT("-Json"),			TEXT("Optionally write output to the given JSON file")),
+		TArgument<FStringView>(TEXT("-TocUrl"),			TEXT("URL of a TOC file to fetch info of (overrides Bucket/ServiceUrl)")),
+		TArgument<FStringView>(TEXT("-TocPath"),		TEXT("Path to a .uondemandtoc file on disk")),
+		TArgument<FStringView>(TEXT("-TocKey"),			TEXT("Filter output by TOC key")),
+		TArgument<FStringView>(TEXT("-BuildVersion"),	TEXT("Filter results a particular build version")),
+		TArgument<FStringView>(TEXT("-TargetPlatform"),	TEXT("Only show results for a given platform")),
+		TArgument<FStringView>(TEXT("-ChunkKeys"),		TEXT("Path to write JSON file of chunk keys to")),
+		TArgument<bool>(TEXT("-AllFilters"),			TEXT("Enabled all filters in JSON output")),
+		TArgument<bool>(TEXT("-ContainerHeader"),		TEXT("Add the container header info in JSON")),
+		TArgument<bool>(TEXT("-TocEntries"),			TEXT("Include TOC entries in JSON output")),
+		TArgument<bool>(TEXT("-BlockSizes"),			TEXT("Write block sizes in JSON output")),
+		TArgument<bool>(TEXT("-BlockHashes"),			TEXT("Add block hash values in JSON output")),
+
+		TArgument<FStringView>(TEXT("-BucketPrefix"),	TEXT("Bucket path to enumerate")),
+		TArgument<FStringView>(TEXT("-BucketPath"),		TEXT("Alternative way to set Bucket/BucketPrefix")),
+		S3Arguments,
+	}
+);
+
+} // namespace UE::IoStore::Tool
diff --git a/Engine/Source/Developer/IoStore/OnDemandUtilities/Private/Main.cpp b/Engine/Source/Developer/IoStore/OnDemandUtilities/Private/Main.cpp
new file mode 100644
--- /dev/null
+++ b/Engine/Source/Developer/IoStore/OnDemandUtilities/Private/Main.cpp
@@ -0,0 +1,25 @@
+// Copyright Epic Games, Inc. All Rights Reserved.
+
+#include "Command.h"
+
+namespace UE::IoStore::Tool {
+
+////////////////////////////////////////////////////////////////////////////////
+FArgumentSet S3Arguments = {
+	TArgument<FStringView>(TEXT("-ServiceUrl"),				TEXT("S3 service URL")),
+	TArgument<FStringView>(TEXT("-Bucket"),					TEXT("S3 bucket name")),
+	TArgument<FStringView>(TEXT("-Region"),					TEXT("S3 region code")),
+	TArgument<FStringView>(TEXT("-AccessKey"),				TEXT("S3 access key")),
+	TArgument<FStringView>(TEXT("-SecretKey"),				TEXT("S3 secret key")),
+	TArgument<FStringView>(TEXT("-SessionToken"),			TEXT("S3 session token")),
+	TArgument<FStringView>(TEXT("-CredentialsFile"),		TEXT("S3 credentials file")),
+	TArgument<FStringView>(TEXT("-CredentialsFileKeyName"),	TEXT("S3 credentials to use from file")),
+};
+
+////////////////////////////////////////////////////////////////////////////////
+int32 Main(int32 ArgC, TCHAR* ArgV[])
+{
+	return FCommand::Main(ArgC, ArgV);
+}
+
+} // namespace UE::IoStore::Tool
diff --git a/Engine/Source/Developer/IoStore/OnDemandUtilities/Private/Purl.cpp b/Engine/Source/Developer/IoStore/OnDemandUtilities/Private/Purl.cpp
new file mode 100644
--- /dev/null
+++ b/Engine/Source/Developer/IoStore/OnDemandUtilities/Private/Purl.cpp
@@ -0,0 +1,172 @@
+// Copyright Epic Games, Inc. All Rights Reserved.
+
+#include "Command.h"
+
+#include <Containers/UnrealString.h>
+#include <HAL/FileManager.h>
+#include <HAL/PlatformProcess.h>
+#include <IO/IoBuffer.h>
+#include <IO/Http/Client.h>
+#include <Misc/Paths.h>
+#include <Misc/ScopeExit.h>
+
+#if PLATFORM_WINDOWS
+#	include <Windows/AllowWindowsPlatformTypes.h>
+#		include <winsock2.h>
+#		include <ws2tcpip.h>
+#	include <Windows/HideWindowsPlatformTypes.h>
+#endif
+
+#include <cstdio>
+
+namespace UE::IoStore::Tool
+{
+
+////////////////////////////////////////////////////////////////////////////////
+static void LoadCaCerts()
+{
+	using namespace UE::IoStore::HTTP;
+
+	IFileManager& Ifm = IFileManager::Get();
+	FString PemPath = FPaths::EngineContentDir() / TEXT("Certificates/ThirdParty/cacert.pem");
+	FArchive* Reader = Ifm.CreateFileReader(*PemPath);
+	check(Reader != nullptr)
+
+	uint32 Size = uint32(Reader->TotalSize());
+	FIoBuffer PemData(Size);
+	FMutableMemoryView PemView = PemData.GetMutableView();
+	Reader->Serialize(PemView.GetData(), Size);
+
+	FCertRoots CaRoots(PemData.GetView());
+	check(CaRoots.IsValid());
+
+	FCertRoots::SetDefault(MoveTemp(CaRoots));
+	delete Reader;
+}
+
+////////////////////////////////////////////////////////////////////////////////
+static int32 PurlCommandEntry(const FContext& Context)
+{
+#if PLATFORM_WINDOWS
+	WSADATA WsaData;
+	if (WSAStartup(MAKEWORD(2, 2), &WsaData) == 0x0a9e0493)
+		return 1;
+	ON_SCOPE_EXIT { WSACleanup(); };
+#endif
+
+	using namespace UE::IoStore::HTTP;
+
+	LoadCaCerts();
+
+	FStringView Url = Context.Get<FStringView>(TEXT("Url"));
+	auto AnsiUrl = StringCast<ANSICHAR>(Url.GetData(), Url.Len());
+
+	FString Method(Context.Get<FStringView>(TEXT("-Method"), TEXT("GET")));
+	Method = Method.ToUpper();
+	auto AnsiMethod = StringCast<ANSICHAR>(*Method);
+
+	bool bChunked = false;
+	uint32 ContentSize = 0;
+	auto Sink = [Dest=FIoBuffer(), &ContentSize, &bChunked] (const FTicketStatus& Status) mutable
+	{
+		if (Status.GetId() == FTicketStatus::EId::Response)
+		{
+			FResponse& Response = Status.GetResponse();
+			FAnsiStringView Message = Response.GetStatusMessage();
+			std::printf("%u %.*s\n",
+				Response.GetStatusCode(),
+				Message.Len(),
+				Message.GetData()
+			);
+			Response.ReadHeaders([] (FAnsiStringView Name, FAnsiStringView Value)
+			{
+				std::printf("%.*s: %.*s\n",
+					Name.Len(),
+					Name.GetData(),
+					Value.Len(),
+					Value.GetData()
+				);
+				return true;
+			});
+
+			bChunked = (Response.GetContentLength() == -1);
+			Response.SetDestination(&Dest);
+			return;
+		}
+
+		if (Status.GetId() == FTicketStatus::EId::Content)
+		{
+			ContentSize += uint32(Dest.GetSize());
+			return;
+		}
+
+		if (Status.GetId() == FTicketStatus::EId::Error)
+		{
+			const char* Reason = Status.GetError().Reason;
+			std::printf("ERROR: %s\n", Reason);
+			return;
+		}
+	};
+
+	bool bRedirect = Context.Get<bool>(TEXT("-Redirect"));
+	bool bUseHttp2 = Context.Get<bool>(TEXT("-HttpTwo"));
+	int32 BufSizeKiB = Context.Get<int32>(TEXT("-BufSize"), -1);
+	uint32 ThrottleKiBps = Context.Get<uint32>(TEXT("-Throttle"), 0);
+
+	FConnectionPool::FParams PoolParams = {
+		.RecvBufSize = BufSizeKiB << 10,
+		.ConnectionCount = 1,
+		.HttpVersion = (bUseHttp2 ? EHttpVersion::Two : EHttpVersion::One),
+	};
+	if (!PoolParams.SetHostFromUrl(AnsiUrl))
+	{
+		std::printf("ERROR: invalid url\n");
+		return 1;
+	}
+	FConnectionPool Pool(PoolParams);
+
+	FEventLoop Loop;
+
+	FEventLoop::FRequestParams RequestParams = {
+		.bAutoRedirect = bRedirect,
+	};
+
+	FAnsiStringView AnsiPath = AnsiUrl;
+	for (int32 i = 2; i >= 0; i -= 2)
+	{
+		int32 Index;
+		AnsiPath.FindChar('/', Index);
+		AnsiPath = AnsiPath.Mid(Index + i);
+	}
+
+	FRequest Request = Loop.Request(AnsiMethod, AnsiPath, Pool, &RequestParams);
+	Request.Header("User-Agent", "purl");
+	Loop.Send(MoveTemp(Request), Sink);
+
+	Loop.SetFailTimeout(4000);
+	Loop.Throttle(ThrottleKiBps);
+	while (Loop.Tick(-1))
+	{
+	}
+
+	std::printf("Data: %u bytes\n", ContentSize);
+
+	return 0;
+}
+
+////////////////////////////////////////////////////////////////////////////////
+static FCommand PurlCommand(
+	PurlCommandEntry,
+	TEXT("Purl"),
+	TEXT("Uses IoStore's HTTP client to download a URL"),
+	{
+		TArgument<FStringView>(TEXT("Url"), TEXT("Url to download")),
+		TArgument<FStringView>(TEXT("-Method"), TEXT("Request method")),
+		TArgument<bool>(TEXT("-Redirect"), TEXT("Follow 30x redirects")),
+		TArgument<bool>(TEXT("-HttpTwo"), TEXT("Use HTTP/2")),
+		TArgument<int32>(TEXT("-BufSize"), TEXT("Recv buffer size in KiB")),
+		TArgument<uint32>(TEXT("-Throttle"), TEXT("Bandwidth throttle in KiB/s")),
+	}
+);
+
+} // namespace UE::IoStore::Tool
diff --git a/Engine/Source/Developer/IoStore/OnDemandUtilities/Private/ResaveToc.cpp b/Engine/Source/Developer/IoStore/OnDemandUtilities/Private/ResaveToc.cpp
new file mode 100644
--- /dev/null
+++ b/Engine/Source/Developer/IoStore/OnDemandUtilities/Private/ResaveToc.cpp
@@ -0,0 +1,91 @@
+// Copyright Epic Games, Inc. All Rights Reserved.
+
+#include "Command.h"
+
+#include "HAL/FileManager.h"
+#include "IO/IoStoreOnDemand.h"
+#include "IO/OnDemandToc.h"
+#include "IO/Serialization/OnDemandContainerToc.h"
+#include "Misc/Paths.h"
+
+namespace UE::IoStore::Tool
+{
+
+////////////////////////////////////////////////////////////////////////////////
+static int32 ResaveTocEntry(const FContext& Context)
+{
+	GPrintLogCategory = GPrintLogVerbosity = false;
+	GPrintLogTimes = ELogTimes::None;
+
+	FString FilePath = FString(Context.Get<FStringView>(TEXT("-Path")));
+	if (FilePath.IsEmpty())
+	{
+		UE_LOG(LogIoStoreOnDemand, Error, TEXT("Invalid filename"));
+		return -1;
+	}
+
+	IFileManager& Ifm = IFileManager::Get();
+	
+	int64 LegacyTocSize = 0;
+	FOnDemandToc LegacyToc;
+	{
+		TUniquePtr<FArchive> Ar(Ifm.CreateFileReader(*FilePath));
+		if (Ar.IsValid() == false)
+		{
+			UE_LOG(LogIoStoreOnDemand, Error, TEXT("Failed to open '%s'"), *FilePath);
+			return -1;
+		}
+		*Ar << LegacyToc;
+		LegacyTocSize = int64(Ar->Tell());
+
+		if (Ar->IsError())
+		{
+			UE_LOG(LogIoStoreOnDemand, Error, TEXT("Failed to serialize legacy TOC '%s'"), *FilePath);
+			return -1;
+		}
+	}
+
+	UE_LOG(LogIoStoreOnDemand, Display, TEXT("Loaded legacy TOC '%s' (%.2lf KB)"),
+		*FilePath, double(LegacyTocSize) / 1024.0);
+
+	FString TmpPath = FPaths::ChangeExtension(FilePath, TEXT("tmp"));
+	if (Ifm.Move(*TmpPath, *FilePath) == false)
+	{
+		UE_LOG(LogIoStoreOnDemand, Error, TEXT("Failed to move file '%s' -> '%s'"), *FilePath, *TmpPath);
+		return -1;
+	}
+
+	UE::IoStore::Serialization::FOnDemandTocWriter Writer;
+	if (TIoStatusOr<uint64> Status = Writer.Write(LegacyToc, FilePath); !Status.IsOk())
+	{
+		Ifm.Delete(*FilePath, false, true, true);
+		if (Ifm.Move(*FilePath, *TmpPath) == false)
+		{
+			UE_LOG(LogIoStoreOnDemand, Error, TEXT("Failed to move file '%s' -> '%s'"), *TmpPath, *FilePath);
+		}
+
+		UE_LOG(LogIoStoreOnDemand, Error, TEXT("Failed to resave file '%s', reason: %s"),
+			*FilePath, *Status.Status().ToString());
+		return -1;
+	}
+	else
+	{
+		Ifm.Delete(*TmpPath, false, true, true);
+		UE_LOG(LogIoStoreOnDemand, Display, TEXT("Successfully resaved '%s' (%.2lf KB)"),
+			*FilePath, double(Status.ConsumeValueOrDie()) / 1024.0);
+	}
+
+	return 0;
+}
+
+////////////////////////////////////////////////////////////////////////////////
+static FCommand ResaveTocCommand(
+	ResaveTocEntry,
+	TEXT("ResaveToc"),
+	TEXT("Resaves a serialized on-demand TOC file to the latest format"),
+	{
+		TArgument<FStringView>(TEXT("-Path"), TEXT("Path to a .uondemandtoc file")),
+	}
+);
+
+} // namespace UE::IoStore::Tool
diff --git a/Engine/Source/Developer/IoStore/OnDemandUtilities/Private/Socks.cpp b/Engine/Source/Developer/IoStore/OnDemandUtilities/Private/Socks.cpp
new file mode 100644
--- /dev/null
+++ b/Engine/Source/Developer/IoStore/OnDemandUtilities/Private/Socks.cpp
@@ -0,0 +1,295 @@
+// Copyright Epic Games, Inc. All Rights Reserved.
+
+#if PLATFORM_WINDOWS
+
+#include "CoreGlobals.h"
+#include "Command.h"
+
+#include <Misc/ScopeExit.h>
+
+#if PLATFORM_WINDOWS
+#	include <Windows/AllowWindowsPlatformTypes.h>
+#		include <winsock2.h>
+#		include <ws2tcpip.h>
+#	include <Windows/HideWindowsPlatformTypes.h>
+#endif
+
+#include <atomic>
+#include <cstdio>
+#include <thread>
+
+namespace UE::IoStore::Tool
+{
+
+////////////////////////////////////////////////////////////////////////////////
+struct FSocks5Stats
+{
+	uint32				Source = 0;
+	std::atomic<uint32>	Dest = 0;
+	std::atomic<int32>	Id = 0;
+	std::atomic<uint64>	Counts[2]; // in KiB
+};
+
+////////////////////////////////////////////////////////////////////////////////
+static void Socks5Peer(SOCKET Peer, FSocks5Stats* Stats)
+{
+#ifdef _MSC_VER
+	// MSVC's static analysis doesn't see that 'Result' from recv() is checked
+	// to be the exact size of the destination buffer.
+	#pragma warning(push)
+	#pragma warning(disable : 6385)
+#endif // _MSC_VER
+
+	ON_SCOPE_EXIT
+	{
+		closesocket(Peer);
+		Stats->Id.store(-1, std::memory_order_release);
+	};
+
+	// greeting
+	char Buf[32];
+	for (int32 i = 0; i < 2; ++i)
+	{
+		if (recv(Peer, Buf + i, 1, 0) != 1)
+		{
+			return;
+		}
+	}
+	check(Buf[0] == 5);
+
+	bool bNoAuth = false;
+	for (int32 i = 0, n = Buf[1]; i < n; ++i)
+	{
+		if (recv(Peer, Buf + i, 1, 0) != 1)
+		{
+			return;
+		}
+	}
+	bNoAuth |= (Buf[0] == 0);
+	check(bNoAuth);
+
+	// ack
+	Buf[0] = 5; Buf[1] = 0;
+	check(send(Peer, Buf, 2, 0) == 2);
+
+	// connect req.
+	for (int32 i = 0; i < 10; ++i)
+	{
+		if (recv(Peer, Buf + i, 1, 0) != 1)
+		{
+			return;
+		}
+	}
+	check(Buf[0] == 5);
+	check(Buf[1] == 1);
+	check(Buf[3] == 1);
+
+	uint32 Ip;
+	uint16 Port;
+	std::memcpy(&Ip, Buf + 4, sizeof(Ip));
+	std::memcpy(&Port, Buf + 8, sizeof(Port));
+	Stats->Dest.store(Ip);
+
+	SOCKET Remote = socket(AF_INET, SOCK_STREAM, IPPROTO_TCP);
+	check(Remote != INVALID_SOCKET);
+	ON_SCOPE_EXIT { closesocket(Remote); };
+
+	sockaddr_in Addr = { sizeof(sockaddr_in) };
+	Addr.sin_family = AF_INET;
+	Addr.sin_addr.s_addr = Ip;
+	Addr.sin_port = Port;
+	if (connect(Remote, &(sockaddr&)Addr, sizeof(Addr)) != 0)
+	{
+		return;
+	}
+
+	std::memset(Buf, 0, sizeof(Buf));
+	Buf[0] = 5;
+	Buf[3] = 1;
+	check(send(Peer, Buf, 10, 0) == 10);
+
+	static const uint32 BufferSize = 1 << 20;
+	char* Buffer = (char*)malloc(BufferSize);
+	ON_SCOPE_EXIT { free(Buffer); };
+
+	while (true)
+	{
+		pollfd Polls[2] = { { Peer, POLLIN }, { Remote, POLLIN } };
+		int32 Result = WSAPoll(Polls, 2, 100);
+		check(Result >= 0);
+		if (Result == 0)
+		{
+			continue;
+		}
+
+		auto Xfer = [&] (SOCKET In, SOCKET Out)
+		{
+			int32 Size = recv(In, Buffer, BufferSize, 0);
+			for (int32 Remaining = Size; Remaining > 0;)
+			{
+				Result = send(Out, Buffer, Size, 0);
+				if (Result <= 0)
+				{
+					return;
+				}
+				Remaining -= Size;
+			}
+
+			int32 Index = (Out == Remote);
+			Stats->Counts[Index].fetch_add(Size >> 10, std::memory_order_relaxed);
+		};
+
+		if ((Polls[0].revents & ~POLLIN) || (Polls[1].revents & ~POLLIN))
+		{
+			break;
+		}
+
+		if (Polls[0].revents & POLLIN)
+		{
+			Xfer(Peer, Remote);
+		}
+
+		if (Polls[1].revents & POLLIN)
+		{
+			Xfer(Remote, Peer);
+		}
+	}
+
+#ifdef _MSC_VER
+	#pragma warning(pop)
+#endif // _MSC_VER
+}
+
+////////////////////////////////////////////////////////////////////////////////
+static void Socks5Server(uint32 Port)
+{
+	auto PrintIp = [] (uint32 Ip4)
+	{
+		std::printf(
+			"%3u.%3u.%3u.%3u",
+			(Ip4 & 0x00'00'00'ff) >>  0,
+			(Ip4 & 0x00'00'ff'00) >>  8,
+			(Ip4 & 0x00'ff'00'00) >> 16,
+			(Ip4 & 0xff'00'00'00) >> 24
+		);
+	};
+
+	TArray<FSocks5Stats*> AllStats;
+	int32 Counter = 0;
+	uint64 SoFar[2] = {};
+
+	std::printf("Listening on 0.0.0.0:%u\n\n", Port);
+
+	SOCKET Sock = socket(AF_INET, SOCK_STREAM, IPPROTO_TCP);
+	check(Sock != INVALID_SOCKET);
+
+	sockaddr_in Addr = { sizeof(sockaddr_in) };
+	Addr.sin_family = AF_INET;
+	Addr.sin_addr.s_addr = 0;
+	Addr.sin_port = htons(uint16(Port));
+	check(bind(Sock, &(sockaddr&)Addr, sizeof(Addr)) == 0);
+	check(listen(Sock, 32) == 0);
+	for (uint32 Tick = 0; !IsEngineExitRequested();)
+	{
+		pollfd Poll = { Sock, POLLIN };
+		int32 Result = WSAPoll(&Poll, 1, 654);
+		if (Result == -1)
+		{
+			break;
+		}
+
+		if (Result == 1)
+		{
+			int32 AddrSize = sizeof(Addr);
+			SOCKET Client = accept(Sock, &(sockaddr&)Addr, &AddrSize);
+
+			FSocks5Stats* PeerStats = new FSocks5Stats();
+			PeerStats->Source = Addr.sin_addr.s_addr;
+			PeerStats->Id = ++Counter;
+			AllStats.Add(PeerStats);
+			std::thread(Socks5Peer, Client, PeerStats).detach();
+			continue;
+		}
+
+		uint64 UpDown[2] = {};
+		uint32 LineClears = 0;
+		for (int32 i = 0, n = AllStats.Num(); i < n; ++i)
+		{
+			FSocks5Stats* Stats = AllStats[i];
+
+			uint64 DownKiB = Stats->Counts[0].load(std::memory_order_relaxed);
+			uint64 UpKiB = Stats->Counts[1].load(std::memory_order_relaxed);
+
+			if (Stats->Id.load(std::memory_order_acquire) == -1)
+			{
+				SoFar[0] += UpKiB;
+				SoFar[1] += DownKiB;
+				delete Stats;
+				AllStats[i] = AllStats.Last();
+				AllStats.Pop();
+				--i, --n, ++LineClears;
+				continue;
+			}
+
+			std::printf("%04d: ", Stats->Id.load(std::memory_order_relaxed));
+			PrintIp(Stats->Source);
+			std::printf("  ->  ");
+			PrintIp(Stats->Dest.load(std::memory_order_relaxed));
+			std::printf(
+				" : d:%9llu u:%9llu KiB\n",
+				DownKiB,
+				UpKiB
+			);
+
+			UpDown[0] += DownKiB;
+			UpDown[1] += UpKiB;
+		}
+
+		for (uint32 i = 0; i < LineClears; ++i)
+		{
+			std::printf("\x1b[2K\n");
+		}
+		std::printf("\x1b[%uF", uint32(AllStats.Num()) + LineClears + 1);
+
+		const char* SignOfLife = ".oOo";
+		std::printf(
+			"[%c] n:%d d:%llu u:%llu\n",
+			SignOfLife[Tick++ & 0x3],
+			Counter,
+			SoFar[0] + UpDown[0],
+			SoFar[1] + UpDown[1]
+		);
+	}
+
+	closesocket(Sock);
+}
+
+////////////////////////////////////////////////////////////////////////////////
+static int32 SocksCommandEntry(const FContext& Context)
+{
+	WSADATA WsaData;
+	if (WSAStartup(MAKEWORD(2, 2), &WsaData) == 0x0a9e0493)
+	{
+		return 1;
+	}
+	ON_SCOPE_EXIT { WSACleanup(); };
+
+	uint32 Port = Context.Get<uint32>(TEXT("-Port"), 24930u);
+	std::thread(Socks5Server, Port).join();
+
+	return 0;
+}
+
+////////////////////////////////////////////////////////////////////////////////
+static FCommand SocksCommand(
+	SocksCommandEntry,
+	TEXT("Socks"),
+	TEXT("Rudimentary SOCKS5 proxy to aid in testing IAS traffic"),
+	{
+		TArgument<uint32>(TEXT("-Port"), TEXT("Port to listen on (default=24930)")),
+	}
+);
+
+} // namespace UE::IoStore::Tool
+
+#endif // PLATFORM_WINDOWS
diff --git a/Engine/Source/Developer/IoStore/OnDemandUtilities/Private/Test.cpp b/Engine/Source/Developer/IoStore/OnDemandUtilities/Private/Test.cpp
new file mode 100644
--- /dev/null
+++ b/Engine/Source/Developer/IoStore/OnDemandUtilities/Private/Test.cpp
@@ -0,0 +1,79 @@
+// Copyright Epic Games, Inc. All Rights Reserved.
+
+#include "Command.h"
+
+#include "Async/TaskGraphInterfaces.h"
+#include "Containers/AnsiString.h"
+#include "Containers/StringConv.h"
+#include "Containers/UnrealString.h"
+#include "Misc/CommandLine.h"
+
+namespace UE::IoStore {
+
+////////////////////////////////////////////////////////////////////////////////
+namespace IasJournaledFileCacheTest { void Tests(const TCHAR*);					}
+namespace HTTP						{ void IasHttpTest(const ANSICHAR*, uint32);}
+
+namespace Tool {
+
+////////////////////////////////////////////////////////////////////////////////
+void CommandTest();
+
+////////////////////////////////////////////////////////////////////////////////
+static void HttpTests(const FContext& Context)
+{
+	FAnsiString TestHost = FAnsiString(Context.Get<FStringView>(TEXT("-Host"), TEXT("localhost")));
+
+	uint32 Seed = Context.Get<uint32>(TEXT("-HttpSeed"), 493);
+	HTTP::IasHttpTest(*TestHost, Seed);
+}
+
+////////////////////////////////////////////////////////////////////////////////
+static void CacheTests(const FContext& Context)
+{
+	FString CacheDirStr = FString(Context.Get<FStringView>(TEXT("-Dir")));
+
+	const TCHAR* CacheDir = CacheDirStr.IsEmpty() ? nullptr : *CacheDirStr;
+	IasJournaledFileCacheTest::Tests(CacheDir);
+}
+
+////////////////////////////////////////////////////////////////////////////////
+static int32 TestCommandEntry(const FContext& Context)
+{
+	FStringView Only = Context.Get<FStringView>(TEXT("-Only"));
+
+	if (Only.IsEmpty())
+	{
+		CommandTest();
+	}
+
+#if IS_PROGRAM
+	if (Only.IsEmpty() || Only == TEXT("cache"))
+	{
+		CacheTests(Context);
+	}
+#endif //IS_PROGRAM
+
+	if (Only.IsEmpty() || Only == TEXT("http"))
+	{
+		HttpTests(Context);
+	}
+
+	return 0;
+}
+
+////////////////////////////////////////////////////////////////////////////////
+static FCommand TestCommand(
+	TestCommandEntry,
+	TEXT("Test"),
+	TEXT("Run IAS tests"),
+	{
+		TArgument<FStringView>(TEXT("-Host"), TEXT("Host of the HTTP test server")),
+		TArgument<FStringView>(TEXT("-Dir"), TEXT("Primary directory to use for cache tests")),
+		TArgument<FStringView>(TEXT("-Only"), TEXT("Only run a particular test (http|cache)")),
+		TArgument<uint32>(TEXT("-HttpSeed"), TEXT("Integer value to seed test HTTP server")),
+	}
+);
+
+} // namespace Tool
+} // namespace UE::IoStore
diff --git a/Engine/Source/Developer/IoStore/OnDemandUtilities/Private/Upload.cpp b/Engine/Source/Developer/IoStore/OnDemandUtilities/Private/Upload.cpp
new file mode 100644
--- /dev/null
+++ b/Engine/Source/Developer/IoStore/OnDemandUtilities/Private/Upload.cpp
@@ -0,0 +1,1045 @@
+// Copyright Epic Games, Inc. All Rights Reserved.
+
+#include "Command.h"
+#include "Common.h"
+
+#include "Async/Async.h"
+#include "Async/Mutex.h"
+#include "Containers/Queue.h"
+#include "HAL/FileManager.h"
+#include "IO/IoBuffer.h"
+#include "IO/IoChunkEncoding.h"
+#include "IO/IoStore.h"
+#include "IO/IoStoreOnDemand.h"
+#include "Misc/ConfigCacheIni.h"
+#include "Misc/FileHelper.h"
+#include "Misc/KeyChainUtilities.h"
+#include "Misc/Paths.h"
+#include "Misc/PathViews.h"
+#include "S3/S3Client.h"
+#include "Serialization/CompactBinarySerialization.h"
+#include "Serialization/CompactBinaryWriter.h"
+#include "Serialization/MemoryReader.h"
+#include "Serialization/LargeMemoryWriter.h"
+#include "UploadQueue.h"
+
+namespace UE::IoStore::Tool
+{
+
+////////////////////////////////////////////////////////////////////////////////
+struct FUploadParams
+{
+	FString ServiceUrl;
+	FString DistributionUrl;
+	FString FallbackUrl;
+	FString Bucket;
+	FString BucketPrefix;
+	FString Region; 
+	FString AccessKey;
+	FString SecretKey;
+	FString SessionToken;
+	FString CredentialsFile;
+	FString CredentialsFileKeyName;
+	FString BuildVersion;
+	FString TargetPlatform;
+	FString EncryptionKeyName;
+	FString HostGroupName;
+	FString GlobalContainerPath;
+	int32 MaxConcurrentUploads = 16;
+	int32 MaxTocListCount = 10000;
+	int32 MaxTocDownloadCount = 100;
+	EOnDemandTocFlags TocFlags = EOnDemandTocFlags::None;
+
+	bool bDeleteContainerFiles = true;
+	bool bDeletePakFiles = true;
+	bool bPerContainerTocs = false;
+	bool bIgnoreContainerFlags = false;
+	bool bIgnoreContainerHeader = false; 
+
+	/** If we should write out the .iochunktoc to disk as well as uploading it. */
+	bool bWriteTocToDisk = false;
+	/** Where the .iochunktoc file should be written out. */
+	FString TocOutputDir;
+
+	FIoStatus Validate() const;
+};
+
+FIoStatus FUploadParams::Validate() const
+{
+	if (GlobalContainerPath.IsEmpty())
+	{
+		if (bWriteTocToDisk && TocOutputDir.IsEmpty())
+		{
+			return FIoStatus(EIoErrorCode::InvalidParameter, TEXT("Cmdline param 'WriteToDisk' requires a valid 'ConfigFilePath' param as well"));
+		}
+	}
+
+	if (!AccessKey.IsEmpty() && SecretKey.IsEmpty())
+	{
+		return FIoStatus(EIoErrorCode::InvalidParameter, TEXT("Invalid secret key"));
+	}
+	else if (AccessKey.IsEmpty() && !SecretKey.IsEmpty())
+	{
+		return FIoStatus(EIoErrorCode::InvalidParameter, TEXT("Invalid access key"));
+	}
+
+	if (!CredentialsFile.IsEmpty() && CredentialsFileKeyName.IsEmpty())
+	{
+		return FIoStatus(EIoErrorCode::InvalidParameter, TEXT("Invalid credential file key name"));
+	}
+
+	if (ServiceUrl.IsEmpty() && Region.IsEmpty())
+	{
+		return FIoStatus(EIoErrorCode::InvalidParameter, TEXT("Service URL or AWS region needs to be specified"));
+	}
+
+	return FIoStatus::Ok;
+}
+
+////////////////////////////////////////////////////////////////////////////////
+struct FUploadResult
+{
+	/** Hash of the toc when written as a binary blob */
+	FIoHash TocHash;
+	/** Url of the service that the toc was uploaded too */
+	FString ServiceUrl;
+	/** Path of the toc on the service */
+	FString TocPath;
+	/** Size (in bytes) of the toc when written as a binary blob */
+	uint64 TocSize = 0;
+};
+
+////////////////////////////////////////////////////////////////////////////////
+static FUploadParams BuildUploadParams(const FContext& Context)
+{
+	static constexpr int32 MinTocListCount		= 1000;
+	static constexpr int32 MaxTocListCount		= 100000;
+	static constexpr int32 MaxTocDownloadCount	= 10000;
+
+	FUploadParams Ret;
+
+	Ret.ServiceUrl				= Context.Get<FStringView>(TEXT("-ServiceUrl"),				Ret.ServiceUrl);
+	Ret.Bucket					= Context.Get<FStringView>(TEXT("-Bucket"),					Ret.Bucket);
+	Ret.BucketPrefix			= Context.Get<FStringView>(TEXT("-BucketPrefix"),			Ret.BucketPrefix);
+	Ret.Region					= Context.Get<FStringView>(TEXT("-Region"),					Ret.Region);
+	Ret.AccessKey				= Context.Get<FStringView>(TEXT("-AccessKey"),				Ret.AccessKey);
+	Ret.SecretKey				= Context.Get<FStringView>(TEXT("-SecretKey"),				Ret.SecretKey);
+	Ret.SessionToken			= Context.Get<FStringView>(TEXT("-SessionToken"),			Ret.SessionToken);
+	Ret.CredentialsFile			= Context.Get<FStringView>(TEXT("-CredentialsFile"),		Ret.CredentialsFile);
+	Ret.CredentialsFileKeyName	= Context.Get<FStringView>(TEXT("-CredentialsFileKeyName"),	Ret.CredentialsFileKeyName);
+	Ret.BuildVersion			= Context.Get<FStringView>(TEXT("-BuildVersion"),			Ret.BuildVersion);
+	Ret.TargetPlatform			= Context.Get<FStringView>(TEXT("-TargetPlatform"),			Ret.TargetPlatform);
+	Ret.HostGroupName			= Context.Get<FStringView>(TEXT("-HostGroupName"),			Ret.HostGroupName);
+	Ret.GlobalContainerPath		= Context.Get<FStringView>(TEXT("-GlobalContainerPath"),	Ret.GlobalContainerPath);
+	Ret.bWriteTocToDisk			= Context.Get<bool>(TEXT("-WriteTocToDisk"),				Ret.bWriteTocToDisk);
+	Ret.bPerContainerTocs		= Context.Get<bool>(TEXT("-PerContainerTocs"),				Ret.bPerContainerTocs);
+	Ret.bIgnoreContainerHeader	= Context.Get<bool>(TEXT("-IgnoreContainerHeader"),			Ret.bIgnoreContainerHeader);
+	Ret.MaxConcurrentUploads	= Context.Get<int32>(TEXT("-MaxConcurrentUploads"),			Ret.MaxConcurrentUploads);
+	Ret.MaxTocListCount			= Context.Get<int32>(TEXT("-MaxTocListCount"),				Ret.MaxTocListCount);
+	Ret.MaxTocDownloadCount		= Context.Get<int32>(TEXT("-MaxTocDownloadCount"),			Ret.MaxTocDownloadCount);
+
+	const bool bStreamOnDemand	= Context.Get<bool>(TEXT("-StreamOnDemand"),				false);
+	const bool bInstallOnDemand	= Context.Get<bool>(TEXT("-InstallOnDemand"),				false);
+
+	if (bStreamOnDemand)
+	{
+		Ret.TocFlags |= EOnDemandTocFlags::StreamOnDemand;
+	}
+	else if (bInstallOnDemand)
+	{
+		Ret.TocFlags |= EOnDemandTocFlags::InstallOnDemand;
+	}
+	else
+	{
+		const EOnDemandTocFlags FallbackFlags = EOnDemandTocFlags::StreamOnDemand;
+		UE_LOG(LogIas, Warning, TEXT("No TOC flags was specified, falling back to '%s'"), *LexToString(FallbackFlags));
+		Ret.TocFlags = FallbackFlags;
+	}
+
+	Ret.bDeleteContainerFiles	= !Context.Get<bool>(TEXT("-KeepContainerFiles"),			!Ret.bDeleteContainerFiles);
+	Ret.bDeletePakFiles			= !Context.Get<bool>(TEXT("-KeepPakFiles"),					!Ret.bDeletePakFiles);
+	Ret.bIgnoreContainerFlags	= Context.Get<bool>(TEXT("-IgnoreContainerFlags"),			Ret.bIgnoreContainerFlags);
+
+	if (Ret.bWriteTocToDisk)
+	{
+		FStringView Path = Context.Get<FStringView>(TEXT("-ConfigFilePath"));
+		Ret.TocOutputDir = FPathViews::GetPath(Path);
+
+		Ret.DistributionUrl = Context.Get<FStringView>(TEXT("-DistributionUrl"),	Ret.DistributionUrl);
+		Ret.FallbackUrl = Context.Get<FStringView>(TEXT("-FallbackUrl"),			Ret.FallbackUrl);
+	}
+
+	if (FIoStatus Validation = Ret.Validate(); !Validation.IsOk())
+	{
+		FString Reason = Validation.ToString();
+		Context.Abort(*Reason);
+	}
+
+	Ret.MaxTocListCount		= FMath::Max(FMath::Min(Ret.MaxTocListCount, MaxTocListCount), MinTocListCount);
+	Ret.MaxTocDownloadCount = FMath::Max(FMath::Min(Ret.MaxTocDownloadCount, MaxTocDownloadCount), 1);
+
+	// Make sure this is a directory path
+	if (FPaths::GetExtension(Ret.GlobalContainerPath).IsEmpty() == false)
+	{
+		Ret.GlobalContainerPath	= FPaths::GetPath(Ret.GlobalContainerPath);
+	}
+
+	return Ret;
+}
+
+////////////////////////////////////////////////////////////////////////////////
+static bool SaveArrayToFile(TArrayView64<uint8> Data, const TCHAR* Filename, uint32 WriteFlags = FILEWRITE_None)
+{
+	TUniquePtr<FArchive> Ar = TUniquePtr<FArchive>(IFileManager::Get().CreateFileWriter(Filename, WriteFlags));
+	if (!Ar)
+	{
+		return false;
+	}
+
+	Ar->Serialize(static_cast<void*>(Data.GetData()), Data.Num());
+	Ar->Close();
+
+	return !Ar->IsError() && !Ar->IsCriticalError();
+}
+
+////////////////////////////////////////////////////////////////////////////////
+static FIoStatus WriteContainerFiles(FOnDemandToc& OnDemandToc, const TMap<FIoHash, FString>& UTocPaths)
+{
+	UE_LOG(LogIas, Display, TEXT("Attempting to write out %d '%s' files"), OnDemandToc.Containers.Num(), *FOnDemandToc::FileExt);
+
+	for (FOnDemandTocContainerEntry& Container : OnDemandToc.Containers)
+	{
+		const FString* UTocPath = UTocPaths.Find(Container.UTocHash);
+		if (UTocPath == nullptr)
+		{
+			return FIoStatus(EIoErrorCode::Unknown, FString::Printf(TEXT("Could not find the original path for ondemand container '%s'"), *Container.ContainerName));
+		}
+
+		// Create a new FOnDemandToc to be written out and give it ownership of the FOnDemandTocContainerEntry that we want
+		// associated with it.This way we don't need to write any specialized serialization code to split up the input FOnDemandToc
+		// into one per container.
+		FOnDemandToc ContainerToc;
+		ContainerToc.Header = OnDemandToc.Header;
+		ContainerToc.Meta = OnDemandToc.Meta;
+		ContainerToc.Containers.Emplace(MoveTemp(Container));
+
+		const FString OutputPath = FPathViews::ChangeExtension(*UTocPath, FOnDemandToc::FileExt);
+
+		TUniquePtr<FArchive> Ar(IFileManager::Get().CreateFileWriter(*OutputPath));
+
+		if (!Ar.IsValid())
+		{
+			return FIoStatus(EIoErrorCode::FileOpenFailed, FString::Printf(TEXT("Failed to open '%s' for write"), *OutputPath));
+		}
+
+		// TODO: We should consider adding a hash of the FOnDemandToc that can be computed at runtime on the loaded structure
+		// (to avoid running over the file twice) to verify that nothing was corrupted.
+
+		(*Ar) << ContainerToc;
+
+		FOnDemandTocSentinel Sentinel;
+		(*Ar) << Sentinel;
+
+		if (!Ar->Close())
+		{
+			return FIoStatus(EIoErrorCode::WriteError, FString::Printf(TEXT("Failed to write to '%s'"), *OutputPath));
+		}
+
+		// TODO: Should write an end of file sentinel here for a simple/easy file corruption check
+
+		UE_LOG(LogIas, Display, TEXT("Wrote ondemand container file '%s' to disk"), *OutputPath);
+
+		// Move the container entry back to OnDemandToc in case we want to use the data structure
+		// beyond this point in the future.
+		Container = MoveTemp(ContainerToc.Containers[0]);
+	}
+
+	return FIoStatus::Ok;
+}
+
+////////////////////////////////////////////////////////////////////////////////
+static void GetTocPath(
+	FStringView ServiceUrl,
+	FStringView Bucket,
+	FStringView BucketPrefix,
+	FString& OutServiceUrl,
+	FString& OutTocPath)
+{
+	// The configuration file should specify a service URL without any trailing
+	// host path, i.e. http://{host:port}/{host-path}. Add the trailing path
+	// to the TOC path to form the complete path the TOC from the host, i.e
+	// TocPath={host-path}/{bucket}/{bucket-prefix}/{toc-hash}.iochunktoc
+	
+	if (BucketPrefix.StartsWith(TEXT("/")))
+	{
+		BucketPrefix.RemovePrefix(1);
+	}
+	if (BucketPrefix.EndsWith(TEXT("/")))
+	{
+		BucketPrefix.RemoveSuffix(1);
+	}
+
+	if (ServiceUrl.IsEmpty())
+	{
+		// If the service URL is empty we assume uploading to AWS S3 using the Region parameter
+		// and that we don't need to prefix with the bucket name.
+		OutServiceUrl.Empty();
+		OutTocPath = BucketPrefix;
+	}
+	else
+	{
+		FStringView HostSuffix;
+		int32 Idx = INDEX_NONE;
+		ensure(ServiceUrl.FindChar(':', Idx));
+		const int32 SchemeEnd = Idx + 3;
+
+		Idx = INDEX_NONE;
+		if (ServiceUrl.RightChop(SchemeEnd).FindChar('/', Idx))
+		{
+			OutServiceUrl = ServiceUrl.Left(SchemeEnd + Idx);
+			HostSuffix = ServiceUrl.RightChop(OutServiceUrl.Len() + 1);
+			if (HostSuffix.EndsWith(TEXT("/")))
+			{
+				HostSuffix.RemoveSuffix(1);
+			}
+		}
+		else
+		{
+			OutServiceUrl = ServiceUrl;
+		}
+
+		TStringBuilder<256> Sb;
+		if (!HostSuffix.IsEmpty())
+		{
+			Sb << HostSuffix << TEXT("/");
+		}
+		Sb << Bucket << TEXT("/") << BucketPrefix;
+		OutTocPath = Sb;
+	}
+}
+
+////////////////////////////////////////////////////////////////////////////////
+static void WriteConfigFile(
+	const FUploadParams& UploadParams,
+	const FUploadResult& UploadResult,
+	const FKeyChain& KeyChain)
+{
+	// This string can get quite long so don't bother with any inline storage
+	FStringBuilderBase Sb;
+	Sb << TEXT("[Endpoint]") << TEXT("\r\n");
+
+	if (!UploadParams.DistributionUrl.IsEmpty())
+	{
+		Sb << TEXT("DistributionUrl=\"") << UploadParams.DistributionUrl << TEXT("\"\r\n");
+
+		if (!UploadParams.FallbackUrl.IsEmpty())
+		{
+			Sb << TEXT("FallbackUrl=\"") << UploadParams.FallbackUrl << TEXT("\"\r\n");
+		}
+	}
+	else
+	{
+		Sb << TEXT("ServiceUrl=\"") << UploadResult.ServiceUrl << TEXT("\"\r\n");
+	}
+
+	Sb << TEXT("TocPath=\"") << UploadResult.TocPath << TEXT("\"\r\n");
+
+	// Temporary solution to get replays working with encrypted on demand content
+	{
+		if (!UploadParams.EncryptionKeyName.IsEmpty())
+		{
+			const TCHAR* EncryptionKeyName = *(UploadParams.EncryptionKeyName);
+
+			TOptional<FNamedAESKey> EncryptionKey;
+			for (const TPair<FGuid, FNamedAESKey>& KeyPair: KeyChain.GetEncryptionKeys())
+			{
+				if (KeyPair.Value.Name.Compare(EncryptionKeyName, ESearchCase::IgnoreCase) == 0)
+				{
+					EncryptionKey.Emplace(KeyPair.Value);
+				}
+			}
+
+			if (EncryptionKey)
+			{
+				FString KeyString = FBase64::Encode(EncryptionKey.GetValue().Key.Key, FAES::FAESKey::KeySize);
+				Sb << TEXT("ContentKey=\"") << EncryptionKey.GetValue().Guid.ToString() << TEXT(":") << KeyString << TEXT("\"\r\n");
+			}
+			else
+			{
+				UE_LOG(LogIoStore, Warning, TEXT("Failed to encryption key '%s' in key chain"), EncryptionKeyName);
+			}
+		}
+	}
+
+	FString ConfigFilePath = UploadParams.TocOutputDir / TEXT("IoStoreOnDemand.ini");
+	UE_LOG(LogIoStore, Display, TEXT("Saving on demand config file '%s'"), *ConfigFilePath);
+	if (FFileHelper::SaveStringToFile(Sb.ToString(), *ConfigFilePath) == false)
+	{
+		UE_LOG(LogIoStore, Error, TEXT("Failed to save on demand config file '%s'"), *ConfigFilePath);
+	}
+}
+
+
+
+////////////////////////////////////////////////////////////////////////////////
+static TIoStatusOr<FUploadResult> UploadContainerFiles(
+	const FUploadParams& UploadParams,
+	TConstArrayView<FString> ContainerFiles,
+	const FKeyChain& KeyChain)
+{
+	TMap<FGuid, FAES::FAESKey> EncryptionKeys;
+	for (const TPair<FGuid, FNamedAESKey>& KeyPair: KeyChain.GetEncryptionKeys())
+	{
+		EncryptionKeys.Add(KeyPair.Key, KeyPair.Value.Key);
+	}
+
+	struct FContainerStats
+	{
+		FString ContainerName;
+		uint64 ChunkCount = 0;
+		uint64 TotalBytes = 0;
+		uint64 UploadedChunkCount = 0;
+		uint64 UploadedBytes = 0;
+	};
+	TMap<FString, FContainerStats> ContainerSummary;
+
+	const double StartTime = FPlatformTime::Seconds();
+
+	FS3ClientConfig Config;
+	Config.ServiceUrl = UploadParams.ServiceUrl;
+	Config.Region = UploadParams.Region;
+	
+	FS3ClientCredentials Credentials;
+	if (UploadParams.CredentialsFile.IsEmpty() == false)
+	{
+		UE_LOG(LogIas, Display, TEXT("Loading credentials file '%s'"), *UploadParams.CredentialsFile);
+		FS3CredentialsProfileStore CredentialsStore = FS3CredentialsProfileStore::FromFile(UploadParams.CredentialsFile);
+		if (CredentialsStore.TryGetCredentials(UploadParams.CredentialsFileKeyName, Credentials) == false)
+		{
+			return FIoStatus(EIoErrorCode::InvalidParameter, TEXT("Failed to find valid credentials in credentials file"));
+		}
+		else
+		{
+			UE_LOG(LogIas, Display, TEXT("Found credentials for '%s'"), *UploadParams.CredentialsFileKeyName);
+		}
+	}
+	else
+	{
+		Credentials = FS3ClientCredentials(UploadParams.AccessKey, UploadParams.SecretKey, UploadParams.SessionToken);
+	}
+
+	FS3Client Client(Config, Credentials);
+	FUploadQueue UploadQueue(Client, UploadParams.Bucket, UploadParams.MaxConcurrentUploads);
+
+	if (ContainerFiles.Num() == 0)
+	{
+		return FIoStatus(EIoErrorCode::InvalidParameter, TEXT("No container file(s) specified"));
+	}
+
+	TSet<FIoHash> ExistingChunks;
+	uint64 TotalExistingTocs = 0;
+	uint64 TotalExistingBytes = 0;
+	{
+		TStringBuilder<256> TocsKey;
+		if (!UploadParams.BucketPrefix.IsEmpty())
+		{
+			TocsKey << UploadParams.BucketPrefix << "/";
+		}
+
+		const int32 MaxKeys = 1000; // AWS max value
+		const int32 MaxPaginationCount = FMath::Max(1, UploadParams.MaxTocListCount / MaxKeys);
+		const int32 MaxTocDownloadCount = UploadParams.MaxTocDownloadCount;
+		TArray<FS3Object> Objects;
+		FString Marker;
+
+		for (int32 Pagination = 0; Pagination < MaxPaginationCount; ++Pagination)
+		{
+			UE_LOG(LogIas, Display, TEXT("List #%d/%d '%s/%s/%s' TocCount=%d, MaxTocCount=%d"),
+				Pagination + 1, MaxPaginationCount, *Client.GetConfig().ServiceUrl, *UploadParams.Bucket, TocsKey.ToString(), Objects.Num(), UploadParams.MaxTocListCount);
+			FS3ListObjectResponse Response = Client.ListObjects(FS3ListObjectsRequest
+			{
+				.BucketName = UploadParams.Bucket,
+				.Prefix = TocsKey.ToString(),
+				.Delimiter = TCHAR('/'),
+				.MaxKeys = MaxKeys,
+				.Marker = Marker
+			});
+
+			Marker = MoveTemp(Response.NextMarker);
+
+			for (FS3Object& Object : Response.Objects)
+			{
+				if (Object.Key.EndsWith(TEXT("iochunktoc")))
+				{
+					Objects.Add(MoveTemp(Object));
+				}
+			}
+
+			if (Response.IsOk() == false || Response.Objects.IsEmpty() || Response.bIsTruncated == false || Objects.Num() >= UploadParams.MaxTocListCount)
+			{
+				break;
+			}
+		}
+
+		{
+			// Just to be sure
+			const int32 TotalObjectCount = Objects.Num();
+			TSet<FS3Object> UniqueObjects(Objects);
+			Objects = UniqueObjects.Array();
+			if (Objects.Num() < TotalObjectCount)
+			{
+				UE_LOG(LogIas, Display, TEXT("Removed %d non unique TOC objects"), TotalObjectCount - Objects.Num());
+			}
+		}
+
+		Objects.Sort([](const FS3Object& LHS, const FS3Object& RHS) { return LHS.LastModified > RHS.LastModified; });
+
+		const int32 TocDownloadCount = FMath::Min(Objects.Num(), MaxTocDownloadCount);
+		UE_LOG(LogIas, Display, TEXT("Fetching %d latest TOC file(s) ..."), TocDownloadCount);
+		for (int32 Idx = 0; Idx < TocDownloadCount; ++Idx)
+		{
+			const FS3Object& TocInfo = Objects[Idx];
+			if (TocInfo.Key.EndsWith(TEXT("iochunktoc")) == false)
+			{
+				continue;
+			}
+
+			UE_LOG(LogIas, Display, TEXT("Fetching TOC %d/%d '%s/%s/%s', Size=%llu, LastModified=%s"), Idx + 1, TocDownloadCount, *Client.GetConfig().ServiceUrl, *UploadParams.Bucket, *TocInfo.Key, TocInfo.Size, *TocInfo.LastModifiedText);
+			FS3GetObjectResponse TocResponse = Client.GetObject(FS3GetObjectRequest
+			{
+				UploadParams.Bucket,
+				TocInfo.Key
+			});
+			
+			if (TocResponse.IsOk() == false)
+			{
+				UE_LOG(LogIas, Warning, TEXT("Failed to fetch TOC '%s/%s/%s'"), *Client.GetConfig().ServiceUrl, *UploadParams.Bucket, *TocInfo.Key);
+				continue;
+			}
+
+			FOnDemandToc Toc;
+			FMemoryReaderView Ar(TocResponse.GetBody().GetView());
+			Ar << Toc;
+
+			if (Ar.IsError()) 
+			{
+				Toc = FOnDemandToc{};
+				if (LoadFromCompactBinary(FCbFieldView(TocResponse.GetBody().GetData()), Toc) == false)
+				{
+					UE_LOG(LogIas, Display, TEXT("Failed to load TOC '%s/%s/%s'"), *Client.GetConfig().ServiceUrl, *UploadParams.Bucket, *TocInfo.Key);
+					continue;
+				}
+			}
+
+			for (const FOnDemandTocContainerEntry& ContainerEntry : Toc.Containers)
+			{
+				for (const FOnDemandTocEntry& TocEntry : ContainerEntry.Entries)
+				{
+					bool bIsAlreadyInSet;
+					ExistingChunks.Add(TocEntry.Hash, &bIsAlreadyInSet);
+					if (!bIsAlreadyInSet)
+					{
+						TotalExistingBytes += TocEntry.EncodedSize;
+					}
+				}
+			}
+
+			TotalExistingTocs++;
+		}
+	}
+
+	FString ChunksRelativePath = UploadParams.BucketPrefix.IsEmpty() ? TEXT("Chunks") : FString::Printf(TEXT("%s/Chunks"), *UploadParams.BucketPrefix);
+	ChunksRelativePath.ToLowerInline();
+
+	bool bPerContainerTocsConfigValue = false;
+	GConfig->GetBool(TEXT("Ias"), TEXT("CreatePerContainerTocs"), bPerContainerTocsConfigValue, GEngineIni);
+	const bool bWritePerContainerToc = bPerContainerTocsConfigValue || UploadParams.bPerContainerTocs;
+
+	uint64 TotalUploadedChunks = 0;
+	uint64 TotalUploadedBytes = 0;
+
+	FOnDemandToc OnDemandToc;
+	OnDemandToc.Header.Flags = static_cast<uint32>(UploadParams.TocFlags);
+	OnDemandToc.Header.HostGroupName = UploadParams.HostGroupName;
+	OnDemandToc.Containers.Reserve(ContainerFiles.Num());
+
+	// Map of the .utoc paths that we have created ondemand containers for, indexed by their hash so that the paths can
+	// be looked up later. We do not rely on the filename as we cannot be sure that there won't be duplicate file names
+	// stored in different directories.
+	TMap<FIoHash, FString> UTocPaths;
+
+	TArray<FString> FilesToDelete;
+	for (const FString& Path : ContainerFiles)
+	{
+		FIoStoreReader ContainerFileReader;
+		{
+			FIoStatus Status = ContainerFileReader.Initialize(*FPaths::ChangeExtension(Path, TEXT("")), EncryptionKeys);
+			if (!Status.IsOk())
+			{
+				UE_LOG(LogIas, Error, TEXT("Failed to open container '%s' for reading due to: '%s'"), *Path, *Status.ToString());
+				continue;
+			}
+		}
+
+		if (UploadParams.bIgnoreContainerFlags == false &&
+			EnumHasAnyFlags(ContainerFileReader.GetContainerFlags(), EIoContainerFlags::OnDemand) == false)
+		{
+			continue;
+		}
+		UE_LOG(LogIas, Display, TEXT("Uploading ondemand container '%s'"), *Path);
+
+		const uint32 BlockSize = ContainerFileReader.GetCompressionBlockSize();
+		if (OnDemandToc.Header.BlockSize == 0)
+		{
+			OnDemandToc.Header.BlockSize = ContainerFileReader.GetCompressionBlockSize();
+		}
+		check(OnDemandToc.Header.BlockSize == ContainerFileReader.GetCompressionBlockSize());
+
+		TArray<FIoStoreTocChunkInfo> ChunkInfos;
+		ContainerFileReader.EnumerateChunks([&ChunkInfos](FIoStoreTocChunkInfo&& Info)
+		{ 
+			ChunkInfos.Emplace(MoveTemp(Info));
+			return true;
+		});
+		
+		FOnDemandTocContainerEntry& ContainerEntry = OnDemandToc.Containers.AddDefaulted_GetRef();
+		ContainerEntry.ContainerName = FPaths::GetBaseFilename(Path);
+		ContainerEntry.ContainerId = ContainerFileReader.GetContainerId();
+		ContainerEntry.ContainerFlags = static_cast<uint8>(ContainerFileReader.GetContainerFlags());
+
+		if (EnumHasAnyFlags(ContainerFileReader.GetContainerFlags(), EIoContainerFlags::Encrypted))
+		{
+			ContainerEntry.EncryptionKeyGuid = LexToString(ContainerFileReader.GetEncryptionKeyGuid());
+		}
+		FContainerStats& ContainerStats = ContainerSummary.FindOrAdd(ContainerEntry.ContainerName);
+		
+		ContainerEntry.Entries.Reserve(ChunkInfos.Num());
+
+		for (const FIoStoreTocChunkInfo& ChunkInfo : ChunkInfos)
+		{
+			const bool bDecrypt = false;
+			TIoStatusOr<FIoStoreCompressedReadResult> Status = ContainerFileReader.ReadCompressed(ChunkInfo.Id, FIoReadOptions(), bDecrypt);
+			if (!Status.IsOk())
+			{
+				return Status.Status();
+			}
+
+			FIoStoreCompressedReadResult ReadResult = Status.ConsumeValueOrDie();
+
+			const uint32 BlockOffset = ContainerEntry.BlockSizes.Num();
+			const uint32 BlockCount = ReadResult.Blocks.Num();
+			const FIoHash ChunkHash = FIoHash::HashBuffer(ReadResult.IoBuffer.GetView());
+
+			FMemoryView EncodedBlocks = ReadResult.IoBuffer.GetView();
+			uint64 RawChunkSize = 0;
+			uint64 EncodedChunkSize = 0;
+			for (const FIoStoreCompressedBlockInfo& BlockInfo : ReadResult.Blocks)
+			{
+				check(Align(BlockInfo.CompressedSize, FAES::AESBlockSize) == BlockInfo.AlignedSize);
+				const uint64 EncodedBlockSize = BlockInfo.AlignedSize;
+				ContainerEntry.BlockSizes.Add(uint32(BlockInfo.CompressedSize));
+
+				FMemoryView EncodedBlock = EncodedBlocks.Left(EncodedBlockSize);
+				EncodedBlocks += EncodedBlock.GetSize();
+				ContainerEntry.BlockHashes.Add(FIoChunkEncoding::HashBlock(EncodedBlock));
+
+				EncodedChunkSize += EncodedBlockSize;
+				RawChunkSize += BlockInfo.UncompressedSize;
+
+				if (OnDemandToc.Header.CompressionFormat.IsEmpty() && BlockInfo.CompressionMethod != NAME_None)
+				{
+					OnDemandToc.Header.CompressionFormat = BlockInfo.CompressionMethod.ToString();
+				}
+			}
+
+			if (EncodedChunkSize != ReadResult.IoBuffer.GetSize())
+			{
+				return FIoStatus(EIoErrorCode::ReadError, TEXT("Encoded chunk size does not match buffer"));
+			}
+
+			// At runtime we are limited to MAX_uint32 for chunk lengths to save space and because anything larger than that
+			// is not reasonable to load via IoStoreOnDemand anyway. So we need to check for this now and fail the upload if
+			// there is a chunk that will fail at runtime.
+			// Note that EncodedSize should always be <= RawChunkSize but test both to be safe.
+			if (RawChunkSize > MAX_uint32 || EncodedChunkSize > MAX_uint32)
+			{
+				return FIoStatus(EIoErrorCode::InvalidParameter, WriteToString<512>(*ChunkInfo.FileName, TEXT(": Chunk size should not exceed MAX_uint32")));
+			}
+
+			FOnDemandTocEntry& TocEntry = ContainerEntry.Entries.AddDefaulted_GetRef();
+			TocEntry.ChunkId = ChunkInfo.Id;
+			TocEntry.Hash = ChunkHash;
+			TocEntry.RawSize = RawChunkSize;
+			TocEntry.EncodedSize = EncodedChunkSize;
+			TocEntry.BlockOffset = BlockOffset;
+			TocEntry.BlockCount = BlockCount;
+
+			ContainerStats.ChunkCount++;
+			ContainerStats.TotalBytes += EncodedChunkSize;
+
+			if (ExistingChunks.Contains(TocEntry.Hash) == false)
+			{
+				const FString HashString = LexToString(ChunkHash);
+
+				TStringBuilder<256> Key;
+				Key << ChunksRelativePath
+					<< TEXT("/") << HashString.Left(2)
+					<< TEXT("/") << HashString
+					<< TEXT(".iochunk");
+
+				if (UploadQueue.Enqueue(Key, ReadResult.IoBuffer) == false)
+				{
+					return FIoStatus(EIoErrorCode::WriteError, TEXT("Failed to upload chunk"));
+				}
+
+				TotalUploadedChunks++;
+				TotalUploadedBytes += EncodedChunkSize;
+				ContainerStats.UploadedChunkCount++;
+				ContainerStats.UploadedBytes += EncodedChunkSize;
+			}
+		}
+
+		if (UploadParams.bIgnoreContainerHeader)
+		{
+			UE_LOG(LogIoStoreOnDemand, Display, TEXT("Ignoring container header for '%s'"), *Path)
+		}
+		else
+		{
+			const bool bDecrypt = false;
+			const FIoChunkId ChunkId = CreateContainerHeaderChunkId(ContainerFileReader.GetContainerId());
+			TIoStatusOr<FIoStoreCompressedReadResult> ReadResult = ContainerFileReader.ReadCompressed(ChunkId, FIoReadOptions(), bDecrypt);
+			if (ReadResult.IsOk())
+			{
+				const FIoBuffer& Chunk = ReadResult.ValueOrDie().IoBuffer;
+				ContainerEntry.Header = MakeArrayView<const uint8>(Chunk.GetData(), static_cast<int32>(Chunk.GetSize()));
+			}
+			else
+			{
+				const EIoErrorCode Err = ReadResult.Status().GetErrorCode();
+				if (Err != EIoErrorCode::UnknownChunkID && Err != EIoErrorCode::NotFound)
+				{
+					return ReadResult.Status();
+				}
+			}
+		}
+
+		{
+			const FString UTocFilePath = FPaths::ChangeExtension(Path, TEXT(".utoc"));
+			TArray<uint8> Buffer;
+			if (!FFileHelper::LoadFileToArray(Buffer, *UTocFilePath))
+			{
+				return FIoStatus(EIoErrorCode::ReadError, TEXT("Failed to upload .utoc file"));
+			}
+
+			ContainerEntry.UTocHash = FIoHash::HashBuffer(Buffer.GetData(), Buffer.Num());
+			TStringBuilder<256> Key;
+			if (UploadParams.BucketPrefix.IsEmpty() == false)
+			{
+				Key << UploadParams.BucketPrefix.ToLower() << TEXT("/");
+			}
+			Key << LexToString(ContainerEntry.UTocHash) << TEXT(".utoc");
+
+			const FS3PutObjectResponse Response = Client.TryPutObject(
+				FS3PutObjectRequest{UploadParams.Bucket, Key.ToString(), MakeMemoryView(Buffer.GetData(), Buffer.Num())});
+
+			if (!Response.IsOk())
+			{
+				return FIoStatus(EIoErrorCode::WriteError, FString::Printf(TEXT("Failed to upload '%s' (%s)"), *UTocFilePath, *Response.GetErrorStatus()));
+			}
+			
+			UTocPaths.Add(ContainerEntry.UTocHash, UTocFilePath);
+
+			UE_LOG(LogIas, Display, TEXT("Uploaded '%s'"), *UTocFilePath);	
+		}
+
+		if (UploadParams.bDeleteContainerFiles)
+		{
+			FilesToDelete.Add(Path);
+			ContainerFileReader.GetContainerFilePaths(FilesToDelete);
+
+			// We need the pak files in order to mount OnDemand toc files!
+			if (UploadParams.bDeletePakFiles && !bWritePerContainerToc)
+			{
+				FilesToDelete.Add(FPaths::ChangeExtension(Path, TEXT(".pak")));
+				FilesToDelete.Add(FPaths::ChangeExtension(Path, TEXT(".sig")));
+			}
+		}
+	}
+
+	if (OnDemandToc.Containers.IsEmpty())
+	{
+		return FIoStatus(EIoErrorCode::InvalidParameter, TEXT("No container file(s) marked as on demand"));
+	}
+
+	if (!UploadQueue.Flush())
+	{
+		return FIoStatus(EIoErrorCode::WriteError, TEXT("Failed to upload chunk(s)"));
+	}
+
+	FString ServiceUrl;
+	GetTocPath(UploadParams.ServiceUrl, UploadParams.Bucket, UploadParams.BucketPrefix, ServiceUrl, OnDemandToc.Header.ChunksDirectory);
+
+	OnDemandToc.Meta.EpochTimestamp = FDateTime::Now().ToUnixTimestamp();
+	OnDemandToc.Meta.BuildVersion = UploadParams.BuildVersion;
+	OnDemandToc.Meta.TargetPlatform = UploadParams.TargetPlatform;
+
+	FUploadResult UploadResult;
+	{
+		FLargeMemoryWriter Ar;
+		Ar << OnDemandToc;
+
+		UploadResult.TocHash = FIoHash::HashBuffer(Ar.GetView());
+		const FString TocFilename = LexToString(UploadResult.TocHash) + TEXT(".iochunktoc");
+
+		TStringBuilder<256> Key;
+		if (!UploadParams.BucketPrefix.IsEmpty())
+		{
+			Key << UploadParams.BucketPrefix.ToLower() << TEXT("/");
+		}
+		Key << TocFilename;
+
+		UploadResult.ServiceUrl = ServiceUrl;
+		UploadResult.TocPath = OnDemandToc.Header.ChunksDirectory / TocFilename;
+		UploadResult.TocSize = Ar.TotalSize();
+
+		const FS3PutObjectResponse Response = Client.TryPutObject(FS3PutObjectRequest{UploadParams.Bucket, Key.ToString(), Ar.GetView()});
+		if (Response.IsOk())
+		{
+			UE_LOG(LogIas, Display, TEXT("Uploaded on demand TOC '%s/%s/%s'"), *Client.GetConfig().ServiceUrl, *UploadParams.Bucket, Key.ToString());
+		}
+		else
+		{
+			UE_LOG(LogIas, Warning, TEXT("Failed to upload TOC '%s/%s/%s' (%s)"), *Client.GetConfig().ServiceUrl, *UploadParams.Bucket, Key.ToString(), *Response.GetErrorStatus());
+			return FIoStatus(EIoErrorCode::WriteError, TEXT("Failed to upload TOC"));
+		}
+
+		// Write a single global.uondemandtoc next to global.utoc or in the Cloud directory (legacy)
+		{
+			TStringBuilder<512> OnDemandTocFilePath;
+			if (UploadParams.GlobalContainerPath.IsEmpty() == false)
+			{
+				FPathViews::Append(OnDemandTocFilePath, UploadParams.GlobalContainerPath, TEXT("global.uondemandtoc"));
+			}
+			else if (UploadParams.bWriteTocToDisk)
+			{
+				FPathViews::Append(OnDemandTocFilePath, UploadParams.TocOutputDir, UploadResult.TocHash);
+				OnDemandTocFilePath << TEXT(".iochunktoc");
+			}
+
+			if (OnDemandTocFilePath.Len() > 0)
+			{
+				FPathViews::NormalizeFilename(OnDemandTocFilePath);
+				if (SaveArrayToFile(TArrayView64<uint8>(Ar.GetData(), Ar.TotalSize()), OnDemandTocFilePath.ToString()))
+				{
+					UE_LOG(LogIas, Display, TEXT("Saved on-demand TOC file '%s'"), OnDemandTocFilePath.ToString());
+				}
+				else
+				{
+					UE_LOG(LogIas, Error, TEXT("Failed to save on-demand TOC file '%s'"), OnDemandTocFilePath.ToString());
+				}
+			}
+		}
+	}
+
+	if (bWritePerContainerToc)
+	{
+		// Write out separate .uondemandtoc files, one per.utoc containing ondemand data.
+		FIoStatus Result = WriteContainerFiles(OnDemandToc, UTocPaths);
+		if (!Result.IsOk())
+		{
+			return Result;
+		}
+	}
+
+	for (const FString& Path : FilesToDelete)
+	{
+		if (IFileManager::Get().FileExists(*Path))
+		{
+			UE_LOG(LogIas, Display, TEXT("Deleting '%s'"), *Path); 
+			IFileManager::Get().Delete(*Path);
+		}
+	}
+
+	if (!UploadParams.TocOutputDir.IsEmpty())
+	{
+		WriteConfigFile(UploadParams, UploadResult, KeyChain);
+	}
+
+	// Generate test chunks from 1KiB - 32MiB
+	{
+		TStringBuilder<256> Key;
+		for (int32 Shft = 0; Shft < 16; ++Shft)
+		{
+			Key.Reset();
+			const int32 ChunkId		= (1 << Shft);
+			const int32 ChunkSize	= ChunkId << 10;
+			Key << ChunksRelativePath << TEXT("/") << ChunkId << TEXT("kib") << TEXT(".iochunk");
+
+			const FS3HeadObjectResponse HeadResponse = Client.HeadObject(FS3HeadObjectRequest{UploadParams.Bucket, Key.ToString()});
+			if (HeadResponse.IsOk())
+			{
+				continue;
+			}
+
+			FIoBuffer Buffer(ChunkSize);
+			ANSICHAR* BufferContent = reinterpret_cast<ANSICHAR*>(Buffer.GetData());
+			FAnsiString SizeText = FAnsiString::Printf("%d", ChunkSize);
+			for (int32 Idx = 0; Idx < Buffer.GetSize(); ++Idx)
+			{
+				BufferContent[Idx] = SizeText[Idx % SizeText.Len()];
+			}
+
+			const FS3PutObjectResponse Response = Client.TryPutObject(
+					FS3PutObjectRequest{UploadParams.Bucket, Key.ToString(), Buffer.GetView()}); 
+			if (Response.IsOk())
+			{
+				UE_LOG(LogIas, Display, TEXT("Uploaded test chunk '%s'"), Key.ToString());
+			}
+			else
+			{
+				UE_LOG(LogIas, Display, TEXT("Failed to uploaded test chunk '%s'"), Key.ToString());
+			}
+		}
+	}
+
+	{
+		const double Duration = FPlatformTime::Seconds() - StartTime;
+		
+		UE_LOG(LogIas, Display, TEXT(""));
+		UE_LOG(LogIas, Display, TEXT("------------------------------------------------- Upload Summary -------------------------------------------------"));
+		UE_LOG(LogIas, Display, TEXT("%-15s: %s"), TEXT("Service URL"), *UploadParams.ServiceUrl);
+		UE_LOG(LogIas, Display, TEXT("%-15s: %s"), TEXT("Bucket"), *UploadParams.Bucket);
+		UE_LOG(LogIas, Display, TEXT("%-15s: %s"), TEXT("TargetPlatform"), *UploadParams.TargetPlatform);
+		UE_LOG(LogIas, Display, TEXT("%-15s: %s"), TEXT("BuildVersion"), *UploadParams.BuildVersion);
+		UE_LOG(LogIas, Display, TEXT("%-15s: %s"), TEXT("HostGroupName"), *UploadParams.HostGroupName);
+		UE_LOG(LogIas, Display, TEXT("%-15s: %s"), TEXT("TOC path"), *UploadResult.TocPath);
+		UE_LOG(LogIas, Display, TEXT("%-15s: %s"), TEXT("TOC flags"), *LexToString(UploadParams.TocFlags));
+		UE_LOG(LogIas, Display, TEXT("%-15s: %.2lf KiB"), TEXT("TOC size"), double(UploadResult.TocSize) / 1024.0);
+		UE_LOG(LogIas, Display, TEXT("%-15s: %.2lf second(s)"), TEXT("Duration"), Duration);
+		UE_LOG(LogIas, Display, TEXT(""));
+		
+		UE_LOG(LogIas, Display, TEXT("%-50s %15s %15s %15s %15s"),
+			TEXT("Container"), TEXT("Chunk(s)"), TEXT("Size (MiB)"), TEXT("Uploaded"), TEXT("Uploaded (MiB)"));
+		UE_LOG(LogIas, Display, TEXT("-------------------------------------------------------------------------------------------------------------------"));
+
+		FContainerStats TotalStats;
+		for (const TTuple<FString, FContainerStats>& Kv : ContainerSummary)
+		{
+			UE_LOG(LogIas, Display, TEXT("%-50s %15llu %15.2lf %15llu %15.2lf"),
+				*Kv.Key, Kv.Value.ChunkCount, double(Kv.Value.TotalBytes) / 1024.0 / 1024.0, Kv.Value.UploadedChunkCount, double(Kv.Value.UploadedBytes) / 1024.0 / 1024.0);
+			
+			TotalStats.ChunkCount += Kv.Value.ChunkCount;
+			TotalStats.TotalBytes += Kv.Value.TotalBytes;
+			TotalStats.UploadedChunkCount += Kv.Value.UploadedChunkCount;
+			TotalStats.UploadedBytes += Kv.Value.UploadedBytes;
+		}
+		UE_LOG(LogIas, Display, TEXT("-------------------------------------------------------------------------------------------------------------------"));
+		UE_LOG(LogIas, Display, TEXT("%-50s %15llu %15.2lf %15llu %15.2lf"),
+			TEXT("Total"),TotalStats.ChunkCount, double(TotalStats.TotalBytes) / 1024.0 / 1024.0, TotalStats.UploadedChunkCount, double(TotalStats.UploadedBytes) / 1024.0 / 1024.0);
+		UE_LOG(LogIas, Display, TEXT(""));
+		
+		UE_LOG(LogIas, Display, TEXT("%-50s %15s %15s %15s"), TEXT("Bucket"), TEXT("TOC(s)"), TEXT("Chunk(s)"), TEXT("MiB"));
+		UE_LOG(LogIas, Display, TEXT("-------------------------------------------------------------------------------------------------------------------"));
+		UE_LOG(LogIas, Display, TEXT("%-50s %15llu %15d %15.2lf"), TEXT("Existing"), TotalExistingTocs, ExistingChunks.Num(), double(TotalExistingBytes) / 1024.0 / 1024.0);
+		UE_LOG(LogIas, Display, TEXT("%-50s %15llu %15d %15.2lf"), TEXT("Uploaded"), 1, TotalStats.UploadedChunkCount, double(TotalStats.UploadedBytes) / 1024.0 / 1024.0);
+		UE_LOG(LogIas, Display, TEXT("-------------------------------------------------------------------------------------------------------------------"));
+		UE_LOG(LogIas, Display, TEXT("%-50s %15llu %15d %15.2lf"), TEXT("Total"),
+			TotalExistingTocs + 1, ExistingChunks.Num() + TotalStats.UploadedChunkCount, double(TotalExistingBytes + TotalStats.UploadedBytes) / 1024.0 / 1024.0);
+		UE_LOG(LogIas, Display, TEXT(""));
+	}
+	
+	return UploadResult;
+}
+
+////////////////////////////////////////////////////////////////////////////////
+static TArray<FString> GlobContainers(const FContext& Context)
+{
+	const TCHAR* GlobPattern = Context.Get<FStringView>(TEXT("ContainerGlob")).GetData();
+
+	TArray<FString> Ret;
+
+	if (IFileManager::Get().FileExists(GlobPattern))
+	{
+		Ret.Add(GlobPattern);
+	}
+	else if (IFileManager::Get().DirectoryExists(GlobPattern))
+	{
+		FString Directory = GlobPattern;
+		FPaths::NormalizeDirectoryName(Directory);
+
+		TArray<FString> FoundContainerFiles;
+		IFileManager::Get().FindFiles(FoundContainerFiles, *(Directory / TEXT("*.utoc")), true, false);
+
+		for (const FString& Filename : FoundContainerFiles)
+		{
+			Ret.Emplace(Directory / Filename);
+		}
+	}
+	else
+	{
+		FString Directory = FPaths::GetPath(GlobPattern);
+		FPaths::NormalizeDirectoryName(Directory);
+
+		TArray<FString> FoundContainerFiles;
+		IFileManager::Get().FindFiles(FoundContainerFiles, GlobPattern, true, false);
+
+		for (const FString& Filename : FoundContainerFiles)
+		{
+			Ret.Emplace(Directory / Filename);
+		}
+	}
+
+	return Ret;
+}
+
+////////////////////////////////////////////////////////////////////////////////
+static int32 UploadCommandEntry(const FContext& Context)
+{
+	TArray<FString> Containers = GlobContainers(Context);
+	FKeyChain KeyChain = Common::LoadCryptoKeys(Context);
+	FUploadParams Params = BuildUploadParams(Context);
+
+	TIoStatusOr<FUploadResult> Result = UploadContainerFiles(Params, Containers, KeyChain);
+	if (!Result.IsOk())
+	{
+		FString Reason = Result.Status().ToString();
+		Context.Abort(*Reason);
+	}
+
+	return 0;
+}
+
+////////////////////////////////////////////////////////////////////////////////
+extern FArgumentSet S3Arguments;
+
+static FCommand UploadCommand(
+	UploadCommandEntry,
+	TEXT("Upload"),
+	TEXT("Used to upload IoStore containers to the cloud and convert to on-demand"),
+	{
+		TArgument<FStringView>(TEXT("ContainerGlob"),			TEXT("Path globbed to discover input containers")),
+		TArgument<FStringView>(TEXT("-CryptoKeys"),				TEXT("JSON-format keyring for input containers")),
+		TArgument<FStringView>(TEXT("-BuildVersion"),			TEXT("Optional build version to embed it TOC")),
+		TArgument<FStringView>(TEXT("-TargetPlatform"),			TEXT("If given, embedded in the output TOC")),
+		TArgument<FStringView>(TEXT("-ConfigFilePath"),			TEXT("Path to the config file to write runtime parameters to")),
+		TArgument<FStringView>(TEXT("-DistributionUrl"),		TEXT("URL for IAS to use at runtime. Output to config file")),
+		TArgument<FStringView>(TEXT("-FallbackUrl"),			TEXT("Alternative fallack for DistributionUrl")),
+		TArgument<FStringView>(TEXT("-BucketPrefix"),			TEXT("Path to prefix to bucket objects")),
+		TArgument<FStringView>(TEXT("-HostGroupName"),			TEXT("Host group name or URL")),
+		TArgument<FStringView>(TEXT("-GlobalContainerPath"),	TEXT("Directory path to the content PAK folder")),
+		TArgument<bool>(TEXT("-KeepContainerFiles"),			TEXT("Do not delete container files after upload")),
+		TArgument<bool>(TEXT("-KeepPakFiles"),					TEXT("Do not delete the springboard pak files")),
+		TArgument<bool>(TEXT("-WriteTocToDisk"),				TEXT("Output the TOC to disk as well as uploading")),
+		TArgument<bool>(TEXT("-PerContainerTocs"),				TEXT("Whether to generate TOC's for each container file(s)")),
+		TArgument<bool>(TEXT("-IgnoreContainerFlags"),			TEXT("Whether to ignore the OnDemand container flag")),
+		TArgument<bool>(TEXT("-StreamOnDemand"),				TEXT("Set the content to be streamed on-demand")),
+		TArgument<bool>(TEXT("-InstallOnDemand"),				TEXT("Set the content to be installed on-demand")),
+		TArgument<bool>(TEXT("-IgnoreContainerHeader"),			TEXT("Do not serialize the container header as part of the TOC file.")),
+		TArgument<int32>(TEXT("-MaxConcurrentUploads"),			TEXT("Number of simultaneous uploads")),
+		TArgument<int32>(TEXT("-MaxTocListCount"),				TEXT("Maximum number of TOC file(s) to list from the bucket")),
+		TArgument<int32>(TEXT("-MaxTocDownloadCount"),			TEXT("Maximum number of TOC file(s) to download")),
+		S3Arguments,
+	}
+);
+
+} // namespace UE::IoStore::Tool
diff --git a/Engine/Source/Developer/IoStore/OnDemandUtilities/Private/UploadQueue.cpp b/Engine/Source/Developer/IoStore/OnDemandUtilities/Private/UploadQueue.cpp
new file mode 100644
--- /dev/null
+++ b/Engine/Source/Developer/IoStore/OnDemandUtilities/Private/UploadQueue.cpp
@@ -0,0 +1,129 @@
+// Copyright Epic Games, Inc. All Rights Reserved.
+
+#include "UploadQueue.h"
+
+#include <Async/Async.h>
+#include <CoreGlobals.h>
+#include <Misc/OutputDeviceRedirector.h>
+#include <IO/IoStoreOnDemand.h>
+#include "S3/S3Client.h"
+
+namespace UE::IoStore::Tool
+{
+
+
+FUploadQueue::FUploadQueue(UE::FS3Client& InClient, const FString& InBucket, int32 ThreadCount)
+	: Client(InClient)
+	, Bucket(InBucket)
+{
+	ActiveThreadCount = ThreadCount;
+	for (int32 Idx = 0; Idx < ThreadCount; ++Idx)
+	{
+		Threads.Add(AsyncThread([this]()
+			{
+				ThreadEntry();
+			}));
+	}
+}
+
+bool FUploadQueue::Enqueue(FStringView Key, FIoBuffer Payload)
+{
+	if (ActiveThreadCount == 0)
+	{
+		return false;
+	}
+
+	for (;;)
+	{
+		bool bEnqueued = false;
+		{
+			FScopeLock _(&CriticalSection);
+			if (ConcurrentUploads < Threads.Num())
+			{
+				bEnqueued = Queue.Enqueue(FQueueEntry{ FString(Key), Payload });
+			}
+		}
+
+		if (bEnqueued)
+		{
+			WakeUpEvent->Trigger();
+			break;
+		}
+
+		UploadCompleteEvent->Wait();
+	}
+
+	return true;
+}
+
+void FUploadQueue::ThreadEntry()
+{
+	const int32 MaxAttempts = 3;
+	const float RetryTimeouts[] = {0.5f, 1.0f, 2.0f};
+
+	for (;;)
+	{
+		FQueueEntry Entry;
+		bool bDequeued = false;
+		{
+			FScopeLock _(&CriticalSection);
+			bDequeued = Queue.Dequeue(Entry);
+		}
+
+		if (!bDequeued)
+		{
+			if (bCompleteAdding)
+			{
+				break;
+			}
+			WakeUpEvent->Wait();
+			continue;
+		}
+
+		ConcurrentUploads++;
+		FS3PutObjectResponse Response;
+		for (int32 Attempt = 0; Attempt < MaxAttempts; ++Attempt)
+		{
+			if (Response = Client.TryPutObject(FS3PutObjectRequest{ Bucket, Entry.Key, Entry.Payload.GetView() }); Response.IsOk())
+			{
+				break;
+			}
+			FPlatformProcess::Sleep(RetryTimeouts[Attempt]);
+		}
+
+		if (Response.IsOk())
+		{
+			UE_LOG(LogIoStoreOnDemand, Log, TEXT("Uploaded chunk '%s/%s/%s'"), *Client.GetConfig().ServiceUrl, *Bucket, *Entry.Key);
+		}
+		else
+		{
+			TStringBuilder<256> ErrorResponse;
+			Response.GetErrorResponse(ErrorResponse);
+
+			UE_LOG(LogIoStoreOnDemand, Warning, TEXT("Failed to upload chunk '%s/%s/%s' (%s)"), *Client.GetConfig().ServiceUrl, *Bucket, *Entry.Key, ErrorResponse.ToString());
+			ErrorCount++;
+		}
+		ConcurrentUploads--;
+		UploadCompleteEvent->Trigger();
+	}
+
+	ActiveThreadCount--;
+}
+
+bool FUploadQueue::Flush()
+{
+	bCompleteAdding = true;
+	for (int32 Idx = 0; Idx < Threads.Num(); ++Idx)
+	{
+		WakeUpEvent->Trigger();
+	}
+
+	for (TFuture<void>& Thread : Threads)
+	{
+		Thread.Wait();
+	}
+
+	return ErrorCount == 0;
+}
+
+} // namespace UE::IoStore::Tool
diff --git a/Engine/Source/Developer/IoStore/OnDemandUtilities/Private/UploadQueue.h b/Engine/Source/Developer/IoStore/OnDemandUtilities/Private/UploadQueue.h
new file mode 100644
--- /dev/null
+++ b/Engine/Source/Developer/IoStore/OnDemandUtilities/Private/UploadQueue.h
@@ -0,0 +1,53 @@
+// Copyright Epic Games, Inc. All Rights Reserved.
+
+#pragma once
+
+#include "IO/IoBuffer.h"
+
+#include <Async/Future.h>
+#include <Containers/Array.h>
+#include <Containers/Queue.h>
+#include <Containers/StringView.h>
+#include <Containers/UnrealString.h>
+#include <String/LexFromString.h>
+
+namespace UE
+{
+	class FS3Client;
+}
+
+namespace UE::IoStore::Tool
+{
+
+////////////////////////////////////////////////////////////////////////////////
+// s3 upload helper.... 
+class FUploadQueue
+{
+public:
+	FUploadQueue(UE::FS3Client& Client, const FString& Bucket, int32 ThreadCount);
+	bool Enqueue(FStringView Key, FIoBuffer Payload);
+	bool Flush();
+
+private:
+	void ThreadEntry();
+
+	struct FQueueEntry
+	{
+		FString Key;
+		FIoBuffer Payload;
+	};
+
+	UE::FS3Client& Client;
+	TArray<TFuture<void>> Threads;
+	FString Bucket;
+	FCriticalSection CriticalSection;
+	TQueue<FQueueEntry> Queue;
+	FEventRef WakeUpEvent;
+	FEventRef UploadCompleteEvent;
+	std::atomic_int32_t ConcurrentUploads{ 0 };
+	std::atomic_int32_t ActiveThreadCount{ 0 };
+	std::atomic_int32_t ErrorCount{ 0 };
+	std::atomic_bool bCompleteAdding{ false };
+};
+
+} // namespace UE::IoStore::Tool
diff --git a/Engine/Source/Developer/S3Client/Internal/S3/S3Client.h b/Engine/Source/Developer/S3Client/Internal/S3/S3Client.h
--- a/Engine/Source/Developer/S3Client/Internal/S3/S3Client.h
+++ b/Engine/Source/Developer/S3Client/Internal/S3/S3Client.h
@@ -113,7 +113,7 @@ struct FS3Response
 {
 public:
 
-	FS3Response();
+	FS3Response() = default;
 	
 	FS3Response(const FS3Response& Other);
 	FS3Response(FS3Response&& Other);
diff --git a/Engine/Source/Developer/S3Client/Private/S3Client.cpp b/Engine/Source/Developer/S3Client/Private/S3Client.cpp
--- a/Engine/Source/Developer/S3Client/Private/S3Client.cpp
+++ b/Engine/Source/Developer/S3Client/Private/S3Client.cpp
@@ -377,8 +377,6 @@ FS3CredentialsProfileStore FS3CredentialsProfileStore::FromFile(const FString& F
 
 ////////////////////////////////////////////////////////////////////////////////
 
-FS3Response::FS3Response() = default;
-
 FS3Response::FS3Response(const FS3Response& Other) = default;
 
 FS3Response::FS3Response(FS3Response&& Other)
diff --git a/Engine/Source/Programs/IoStoreOnDemand/IasTool.Build.cs b/Engine/Source/Programs/IoStoreOnDemand/IasTool.Build.cs
--- a/Engine/Source/Programs/IoStoreOnDemand/IasTool.Build.cs
+++ b/Engine/Source/Programs/IoStoreOnDemand/IasTool.Build.cs
@@ -10,6 +10,7 @@ public class IasTool : ModuleRules
 
 		PrivateDependencyModuleNames.Add("Core");
 		PrivateDependencyModuleNames.Add("IoStoreOnDemand");
+		PrivateDependencyModuleNames.Add("IoStoreOnDemandUtilities");
 		PrivateDependencyModuleNames.Add("Projects");
 	}
 }
diff --git a/Engine/Source/Programs/IoStoreOnDemand/IasTool.Target.cs b/Engine/Source/Programs/IoStoreOnDemand/IasTool.Target.cs
--- a/Engine/Source/Programs/IoStoreOnDemand/IasTool.Target.cs
+++ b/Engine/Source/Programs/IoStoreOnDemand/IasTool.Target.cs
@@ -12,8 +12,6 @@ public class IasToolTarget : TargetRules
 		LaunchModuleName = "IasTool";
 		IncludeOrderVersion = EngineIncludeOrderVersion.Latest;
 
-		GlobalDefinitions.Add("UE_WITH_IAS_TOOL=1");
-
 		bBuildDeveloperTools = false;
 		bBuildWithEditorOnlyData = false;
 		bCompileAgainstApplicationCore = false;
diff --git a/Engine/Source/Runtime/Experimental/IoStore/OnDemand/Private/Tool/ChunkPlugin.cpp b/Engine/Source/Runtime/Experimental/IoStore/OnDemand/Private/Tool/ChunkPlugin.cpp
deleted file mode 100644
diff --git a/Engine/Source/Runtime/Experimental/IoStore/OnDemand/Private/Tool/Command.cpp b/Engine/Source/Runtime/Experimental/IoStore/OnDemand/Private/Tool/Command.cpp
deleted file mode 100644
diff --git a/Engine/Source/Runtime/Experimental/IoStore/OnDemand/Private/Tool/Command.h b/Engine/Source/Runtime/Experimental/IoStore/OnDemand/Private/Tool/Command.h
deleted file mode 100644
diff --git a/Engine/Source/Runtime/Experimental/IoStore/OnDemand/Private/Tool/Common.cpp b/Engine/Source/Runtime/Experimental/IoStore/OnDemand/Private/Tool/Common.cpp
deleted file mode 100644
diff --git a/Engine/Source/Runtime/Experimental/IoStore/OnDemand/Private/Tool/Common.h b/Engine/Source/Runtime/Experimental/IoStore/OnDemand/Private/Tool/Common.h
deleted file mode 100644
diff --git a/Engine/Source/Runtime/Experimental/IoStore/OnDemand/Private/Tool/ConvertToc.cpp b/Engine/Source/Runtime/Experimental/IoStore/OnDemand/Private/Tool/ConvertToc.cpp
deleted file mode 100644
diff --git a/Engine/Source/Runtime/Experimental/IoStore/OnDemand/Private/Tool/Download.cpp b/Engine/Source/Runtime/Experimental/IoStore/OnDemand/Private/Tool/Download.cpp
deleted file mode 100644
diff --git a/Engine/Source/Runtime/Experimental/IoStore/OnDemand/Private/Tool/Graph.cpp b/Engine/Source/Runtime/Experimental/IoStore/OnDemand/Private/Tool/Graph.cpp
deleted file mode 100644
diff --git a/Engine/Source/Runtime/Experimental/IoStore/OnDemand/Private/Tool/Graph.h b/Engine/Source/Runtime/Experimental/IoStore/OnDemand/Private/Tool/Graph.h
deleted file mode 100644
diff --git a/Engine/Source/Runtime/Experimental/IoStore/OnDemand/Private/Tool/ListTocs.cpp b/Engine/Source/Runtime/Experimental/IoStore/OnDemand/Private/Tool/ListTocs.cpp
deleted file mode 100644
diff --git a/Engine/Source/Runtime/Experimental/IoStore/OnDemand/Private/Tool/Main.cpp b/Engine/Source/Runtime/Experimental/IoStore/OnDemand/Private/Tool/Main.cpp
deleted file mode 100644
diff --git a/Engine/Source/Runtime/Experimental/IoStore/OnDemand/Private/Tool/Purl.cpp b/Engine/Source/Runtime/Experimental/IoStore/OnDemand/Private/Tool/Purl.cpp
deleted file mode 100644
diff --git a/Engine/Source/Runtime/Experimental/IoStore/OnDemand/Private/Tool/ResaveToc.cpp b/Engine/Source/Runtime/Experimental/IoStore/OnDemand/Private/Tool/ResaveToc.cpp
deleted file mode 100644
diff --git a/Engine/Source/Runtime/Experimental/IoStore/OnDemand/Private/Tool/Socks.cpp b/Engine/Source/Runtime/Experimental/IoStore/OnDemand/Private/Tool/Socks.cpp
deleted file mode 100644
diff --git a/Engine/Source/Runtime/Experimental/IoStore/OnDemand/Private/Tool/Test.cpp b/Engine/Source/Runtime/Experimental/IoStore/OnDemand/Private/Tool/Test.cpp
deleted file mode 100644
diff --git a/Engine/Source/Runtime/Experimental/IoStore/OnDemand/Private/Tool/Upload.cpp b/Engine/Source/Runtime/Experimental/IoStore/OnDemand/Private/Tool/Upload.cpp
deleted file mode 100644
diff --git a/Engine/Source/Runtime/Experimental/IoStore/OnDemand/Private/Tool/UploadQueue.cpp b/Engine/Source/Runtime/Experimental/IoStore/OnDemand/Private/Tool/UploadQueue.cpp
deleted file mode 100644
diff --git a/Engine/Source/Runtime/Experimental/IoStore/OnDemand/Private/Tool/UploadQueue.h b/Engine/Source/Runtime/Experimental/IoStore/OnDemand/Private/Tool/UploadQueue.h
deleted file mode 100644
diff --git a/Engine/Source/Runtime/Experimental/IoStore/OnDemandCore/Internal/IO/Serialization/OnDemandContainerToc.h b/Engine/Source/Runtime/Experimental/IoStore/OnDemandCore/Internal/IO/Serialization/OnDemandContainerToc.h
--- a/Engine/Source/Runtime/Experimental/IoStore/OnDemandCore/Internal/IO/Serialization/OnDemandContainerToc.h
+++ b/Engine/Source/Runtime/Experimental/IoStore/OnDemandCore/Internal/IO/Serialization/OnDemandContainerToc.h
@@ -270,16 +270,16 @@ struct FOnDemandContainerTocHeaderView
 			, ContainerEntry(InContainerEntry)
 	{ }
 
-	FUtf8StringView			ContainerName() const { return GetString(ContainerEntry->ContainerName); }
-	FIoContainerId			ContainerId() const { return ContainerEntry->ContainerId; }
-	EContainerFlags			ContainerFlags() const { return static_cast<EContainerFlags>(ContainerEntry->ContainerFlags); }
-	EIoContainerFlags		FileContainerFlags() const;
-	const FGuid&			EncryptionKeyGuid() const { return ContainerEntry->EncryptionKeyGuid; }
-	uint32					BlockSize() const { return ContainerEntry->BlockSize; }
-	UE_API FUtf8StringView	GetString(FOnDemandStringEntry StringEntry) const;
-
-	bool					IsValid() const { return Header != nullptr && StringTable != nullptr && ContainerEntry != nullptr; }
-							operator bool() const { return IsValid(); }
+	FUtf8StringView				ContainerName() const { return GetString(ContainerEntry->ContainerName); }
+	FIoContainerId				ContainerId() const { return ContainerEntry->ContainerId; }
+	EContainerFlags				ContainerFlags() const { return static_cast<EContainerFlags>(ContainerEntry->ContainerFlags); }
+	UE_API EIoContainerFlags	FileContainerFlags() const;
+	const FGuid&				EncryptionKeyGuid() const { return ContainerEntry->EncryptionKeyGuid; }
+	uint32						BlockSize() const { return ContainerEntry->BlockSize; }
+	UE_API FUtf8StringView		GetString(FOnDemandStringEntry StringEntry) const;
+
+	bool						IsValid() const { return Header != nullptr && StringTable != nullptr && ContainerEntry != nullptr; }
+								operator bool() const { return IsValid(); }
 
 private:
 	const FOnDemandTocHeader*		Header = nullptr;
