Remove an unused local variable that was caught by static code analysis once the code was moved.

#rnx

#Synced-CL 49815329

diff --git a/Engine/Source/Developer/IoStore/OnDemandUtilities/Private/Upload.cpp b/Engine/Source/Developer/IoStore/OnDemandUtilities/Private/Upload.cpp
--- a/Engine/Source/Developer/IoStore/OnDemandUtilities/Private/Upload.cpp
+++ b/Engine/Source/Developer/IoStore/OnDemandUtilities/Private/Upload.cpp
@@ -588,7 +588,6 @@ static TIoStatusOr<FUploadResult> UploadContainerFiles(
 		}
 		UE_LOG(LogIas, Display, TEXT("Uploading ondemand container '%s'"), *Path);
 
-		const uint32 BlockSize = ContainerFileReader.GetCompressionBlockSize();
 		if (OnDemandToc.Header.BlockSize == 0)
 		{
 			OnDemandToc.Header.BlockSize = ContainerFileReader.GetCompressionBlockSize();
