Updated IasTool ResaveToc command to support resaving all on-demand container file(s) in a directory.

#rnx

#Synced-CL 49823132

diff --git a/Engine/Source/Developer/IoStore/OnDemandUtilities/Private/ResaveToc.cpp b/Engine/Source/Developer/IoStore/OnDemandUtilities/Private/ResaveToc.cpp
--- a/Engine/Source/Developer/IoStore/OnDemandUtilities/Private/ResaveToc.cpp
+++ b/Engine/Source/Developer/IoStore/OnDemandUtilities/Private/ResaveToc.cpp
@@ -17,65 +17,85 @@ static int32 ResaveTocEntry(const FContext& Context)
 	GPrintLogCategory = GPrintLogVerbosity = false;
 	GPrintLogTimes = ELogTimes::None;
 
-	FString FilePath = FString(Context.Get<FStringView>(TEXT("-Path")));
-	if (FilePath.IsEmpty())
+	IFileManager& Ifm = IFileManager::Get();
+	TArray<FString> FilePaths;
+
+	if (FString Path = FString(Context.Get<FStringView>(TEXT("-Path"))); !Path.IsEmpty())
 	{
-		UE_LOG(LogIoStoreOnDemand, Error, TEXT("Invalid filename"));
-		return -1;
+		FilePaths.Add(MoveTemp(Path));
 	}
-
-	IFileManager& Ifm = IFileManager::Get();
-	
-	int64 LegacyTocSize = 0;
-	FOnDemandToc LegacyToc;
+	else if (FString Directory = FString(Context.Get<FStringView>(TEXT("-Directory"))); !Directory.IsEmpty())
 	{
-		TUniquePtr<FArchive> Ar(Ifm.CreateFileReader(*FilePath));
-		if (Ar.IsValid() == false)
-		{
-			UE_LOG(LogIoStoreOnDemand, Error, TEXT("Failed to open '%s'"), *FilePath);
-			return -1;
-		}
-		*Ar << LegacyToc;
-		LegacyTocSize = int64(Ar->Tell());
+		TArray<FString> Filenames;
+		Ifm.FindFiles(Filenames, *Directory, UE::IoStore::Serialization::FOnDemandFileExt::Toc);
 
-		if (Ar->IsError())
+		for (const FString& Filename : Filenames)
 		{
-			UE_LOG(LogIoStoreOnDemand, Error, TEXT("Failed to serialize legacy TOC '%s'"), *FilePath);
-			return -1;
+			FilePaths.Add(Directory / Filename);
 		}
 	}
 
-	UE_LOG(LogIoStoreOnDemand, Display, TEXT("Loaded legacy TOC '%s' (%.2lf KB)"),
-		*FilePath, double(LegacyTocSize) / 1024.0);
-
-	FString TmpPath = FPaths::ChangeExtension(FilePath, TEXT("tmp"));
-	if (Ifm.Move(*TmpPath, *FilePath) == false)
+	if (FilePaths.IsEmpty())
 	{
-		UE_LOG(LogIoStoreOnDemand, Error, TEXT("Failed to move file '%s' -> '%s'"), *FilePath, *TmpPath);
+		UE_LOG(LogIoStoreOnDemand, Error, TEXT("Specify a filename or a directory."));
 		return -1;
 	}
 
-	UE::IoStore::Serialization::FOnDemandTocWriter Writer;
-	if (TIoStatusOr<uint64> Status = Writer.Write(LegacyToc, FilePath); !Status.IsOk())
+	int32 Result = 0;
+	for (const FString& FilePath : FilePaths)
 	{
-		Ifm.Delete(*FilePath, false, true, true);
-		if (Ifm.Move(*FilePath, *TmpPath) == false)
+		UE::IoStore::FOnDemandToc LegacyToc;
 		{
-			UE_LOG(LogIoStoreOnDemand, Error, TEXT("Failed to move file '%s' -> '%s'"), *TmpPath, *FilePath);
+			TUniquePtr<FArchive> Ar(Ifm.CreateFileReader(*FilePath));
+			if (Ar.IsValid() == false)
+			{
+				UE_LOG(LogIoStoreOnDemand, Error, TEXT("Failed to open '%s'"), *FilePath);
+				Result = -1;
+				continue;
+			}
+			*Ar << LegacyToc;
+
+			if (Ar->IsError())
+			{
+				UE_LOG(LogIoStoreOnDemand, Error, TEXT("Failed to load legacy TOC '%s'"), *FilePath);
+				Result = -1;
+				continue;
+			}
+
+			UE_LOG(LogIoStoreOnDemand, Display, TEXT("Loaded legacy TOC '%s' (%.2lf KB)"),
+				*FilePath, double(Ar->Tell()) / 1024.0);
 		}
 
-		UE_LOG(LogIoStoreOnDemand, Error, TEXT("Failed to resave file '%s', reason: %s"),
-			*FilePath, *Status.Status().ToString());
-		return -1;
-	}
-	else
-	{
-		Ifm.Delete(*TmpPath, false, true, true);
-		UE_LOG(LogIoStoreOnDemand, Display, TEXT("Successfully resaved '%s' (%.2lf KB)"),
-			*FilePath, double(Status.ConsumeValueOrDie()) / 1024.0);
+		FString BackupPath = FPaths::ChangeExtension(FilePath, TEXT("backup"));
+		if (Ifm.Move(*BackupPath, *FilePath) == false)
+		{
+			UE_LOG(LogIoStoreOnDemand, Error, TEXT("Failed to backup file '%s' -> '%s'"), *FilePath, *BackupPath);
+			Result = -1;
+			continue;
+		}
+
+		UE::IoStore::Serialization::FOnDemandTocWriter Writer;
+		if (TIoStatusOr<uint64> Status = Writer.Write(LegacyToc, FilePath); Status.IsOk())
+		{
+			Ifm.Delete(*BackupPath, false, true, true);
+			UE_LOG(LogIoStoreOnDemand, Display, TEXT("Successfully resaved '%s' (%.2lf KB)"),
+				*FilePath, double(Status.ConsumeValueOrDie()) / 1024.0);
+		}
+		else
+		{
+			Ifm.Delete(*FilePath, false, true, true);
+			if (Ifm.Move(*FilePath, *BackupPath) == false)
+			{
+				UE_LOG(LogIoStoreOnDemand, Error, TEXT("Failed to move file '%s' -> '%s'"), *BackupPath, *FilePath);
+			}
+
+			UE_LOG(LogIoStoreOnDemand, Error, TEXT("Failed to resave file '%s', reason: %s"),
+				*FilePath, *Status.Status().ToString());
+			Result = -1;
+		}
 	}
 
-	return 0;
+	return Result;
 }
 
 ////////////////////////////////////////////////////////////////////////////////
@@ -85,6 +105,7 @@ static FCommand ResaveTocCommand(
 	TEXT("Resaves a serialized on-demand TOC file to the latest format"),
 	{
 		TArgument<FStringView>(TEXT("-Path"), TEXT("Path to a .uondemandtoc file")),
+		TArgument<FStringView>(TEXT("-Directory"), TEXT("Directory to one or more .uondemandtoc file(s)")),
 	}
 );
 
