Remove duplicate code and general clean up for the 'ChunkPlugin' command.

#rnx

- Moved code setting up the meta info for the ondemand containers we want to write so we don't have copy pasted code for both toc types.
- Add code documentation for some FChunkOutput members.
- Restored the reversing of the FOnDemandTocContainerEntry array as copying those structures can be quite expensive.

#Synced-CL 49815329

diff --git a/Engine/Source/Developer/IoStore/OnDemandUtilities/Private/ChunkPlugin.cpp b/Engine/Source/Developer/IoStore/OnDemandUtilities/Private/ChunkPlugin.cpp
--- a/Engine/Source/Developer/IoStore/OnDemandUtilities/Private/ChunkPlugin.cpp
+++ b/Engine/Source/Developer/IoStore/OnDemandUtilities/Private/ChunkPlugin.cpp
@@ -259,17 +259,29 @@ struct FChunkParams
 
 struct FChunkOutput
 {
-	FChunkOutput(const FString& HostGroupName)
+	FChunkOutput(const FString& HostGroupName, const FString& Platform, const FString& BuildVersion)
 	{
+		const TCHAR* ChunksDirectory = TEXT(""); // TODO: Should we be setting this to something?
+
 		InstallOnDemandToc.Header.Flags = static_cast<uint32>(EOnDemandTocFlags::InstallOnDemand);
 		InstallOnDemandToc.Header.HostGroupName = HostGroupName;
-		InstallOnDemandToc.Header.ChunksDirectory = TEXT(""); // TODO?
-		//InstallOnDemandToc.Containers.Reserve(ContainerFilenames.Num());
+		InstallOnDemandToc.Header.ChunksDirectory = ChunksDirectory;
+		InstallOnDemandToc.Meta.TargetPlatform = Platform;
+		InstallOnDemandToc.Meta.BuildVersion = BuildVersion;
 
 		StreamOnDemandToc.Header.Flags = static_cast<uint32>(EOnDemandTocFlags::StreamOnDemand);
 		StreamOnDemandToc.Header.HostGroupName = HostGroupName;
-		StreamOnDemandToc.Header.ChunksDirectory = TEXT(""); // TODO?
-		//StreamOnDemandToc.Containers.Reserve(ContainerFilenames.Num());
+		StreamOnDemandToc.Header.ChunksDirectory = ChunksDirectory;
+		StreamOnDemandToc.Meta.TargetPlatform = Platform;
+		StreamOnDemandToc.Meta.BuildVersion = BuildVersion;
+	}
+
+	void Reserve(int32 NumPotentialContainers, int32 NumPackageSets)
+	{
+		InstallOnDemandToc.Containers.Reserve(NumPotentialContainers);
+		StreamOnDemandToc.Containers.Reserve(NumPotentialContainers);
+
+		FoundTagSets.Reserve(NumPackageSets);
 	}
 
 	FOnDemandToc& GetToc(EOnDemandTocFlags Flags)
@@ -284,11 +296,20 @@ struct FChunkOutput
 		}
 	}
 
+	void SetTocsTimestamp(int64 Timestamp)
+	{
+		InstallOnDemandToc.Meta.EpochTimestamp = Timestamp;
+		StreamOnDemandToc.Meta.EpochTimestamp = Timestamp;
+	}
+
 	FOnDemandToc InstallOnDemandToc;
 	FOnDemandToc StreamOnDemandToc;
 
+	/* Tracks the packages found for each tag per container. Is only used for InstallOnDemandToc. */
 	TMap<FString, TArray<FOnDemandTocTagSetPackageList>> FoundTagSets;
-	TArray<FString> ContainerFiles;
+	/* The paths of the container files (.ucas, .utoc etc) that were chunked. */
+	TArray<FString> ContainerFilePaths;
+	/* Statistics about the chunking process. */
 	FChunkPluginStats Stats;
 };
 
@@ -535,8 +556,8 @@ static int32 ChunkContainer(const FString& ContainerUTocPath, EOnDemandTocFlags
 		}
 	}
 
-	Output.ContainerFiles.Add(ContainerUTocPath);						// .utoc
-	ContainerFileReader.GetContainerFilePaths(Output.ContainerFiles);	// .ucas
+	Output.ContainerFilePaths.Add(ContainerUTocPath);						// .utoc
+	ContainerFileReader.GetContainerFilePaths(Output.ContainerFilePaths);	// .ucas
 
 	return 0;
 }
@@ -556,6 +577,7 @@ int32 WriteOnDemandToc(FOnDemandToc& OnDemandToc, const FString& ContainerName,
 	{
 		*Ar << OnDemandToc;
 		Ar->Close();
+
 		if (Ar->IsError())
 		{
 			UE_LOG(LogIoStoreOnDemand, Error, TEXT("Failed to serialize TOC '%s'"), *TocPath);
@@ -564,8 +586,7 @@ int32 WriteOnDemandToc(FOnDemandToc& OnDemandToc, const FString& ContainerName,
 		else
 		{
 			const int64 TocSize = Ar->Tell();
-			UE_LOG(LogIoStoreOnDemand, Display, TEXT("Writing file '%s' (%.2lf KiB)"),
-				*TocPath, double(TocSize) / 1024);
+			UE_LOG(LogIoStoreOnDemand, Display, TEXT("Writing file '%s' (%.2lf KiB)"), *TocPath, double(TocSize) / 1024);
 
 			return 0;
 		}
@@ -727,8 +748,8 @@ static int32 ChunkPluginCommandEntry(const FContext& Context)
 		return -1;
 	}
 
-	FChunkOutput Output(HostGroupName);
-	Output.FoundTagSets.Reserve(Input.Settings.PackageSets.Num());
+	FChunkOutput Output(HostGroupName, Platform, BuildVersion);
+	Output.Reserve(ContainerFilenames.Num(), Input.Settings.PackageSets.Num());
 
 	for (const FString& Filename : ContainerFilenames)
 	{
@@ -779,7 +800,7 @@ static int32 ChunkPluginCommandEntry(const FContext& Context)
 
 		UE_LOG(LogIoStoreOnDemand, Display, TEXT("Deleting chunked  containers..."));
 
-		for (const FString& Path : Output.ContainerFiles)
+		for (const FString& Path : Output.ContainerFilePaths)
 		{
 			if (FileMgr.FileExists(*Path))
 			{
@@ -856,12 +877,10 @@ static int32 ChunkPluginCommandEntry(const FContext& Context)
 	// const FString TocPath = OutputFolder / Filename; // like this yo
 	const FString OnDemandTocOutputFolder = OutputFolder / IoStoreRelativeFolder;
 
+	Output.SetTocsTimestamp(Timestamp);
+
 	if (bSupportInstalledContainers && !Output.InstallOnDemandToc.Containers.IsEmpty())
 	{
-		Output.InstallOnDemandToc.Meta.EpochTimestamp = Timestamp;
-		Output.InstallOnDemandToc.Meta.TargetPlatform = Platform;
-		Output.InstallOnDemandToc.Meta.BuildVersion = BuildVersion;
-
 		int32 Result = WriteOnDemandToc(Output.InstallOnDemandToc, OnDemandTocName, OnDemandTocOutputFolder);
 		if (Result < 0)
 		{
@@ -871,11 +890,7 @@ static int32 ChunkPluginCommandEntry(const FContext& Context)
 
 	if (bSupportStreamingContainers && !Output.StreamOnDemandToc.Containers.IsEmpty())
 	{
-		Output.StreamOnDemandToc.Meta.EpochTimestamp = Timestamp;
-		Output.StreamOnDemandToc.Meta.TargetPlatform = Platform;
-		Output.StreamOnDemandToc.Meta.BuildVersion = BuildVersion;
-
-		// TODO: OnDemandTocName
+		// TODO: Improve how we choose the name here.
 		int32 Result = WriteOnDemandToc(Output.StreamOnDemandToc, OnDemandTocName + TEXT("ondemand"), OnDemandTocOutputFolder);
 		if (Result < 0)
 		{
@@ -890,7 +905,7 @@ static int32 ChunkPluginCommandEntry(const FContext& Context)
 
 		UE_LOG(LogIoStoreOnDemand, Display, TEXT("Replacing .utoc for chunked  containers..."));
 
-		for (const FString& SourcePath : Output.ContainerFiles)
+		for (const FString& SourcePath : Output.ContainerFilePaths)
 		{
 			if (!SourcePath.EndsWith(TEXT(".utoc")))
 			{
