Moved the setting parsing code to it's own method for the ChunkPlugin command

Remove duplicate code and general clean up for the 'ChunkPlugin' command.

#rnx

#Synced-CL 49815329

diff --git a/Engine/Source/Developer/IoStore/OnDemandUtilities/Private/ChunkPlugin.cpp b/Engine/Source/Developer/IoStore/OnDemandUtilities/Private/ChunkPlugin.cpp
--- a/Engine/Source/Developer/IoStore/OnDemandUtilities/Private/ChunkPlugin.cpp
+++ b/Engine/Source/Developer/IoStore/OnDemandUtilities/Private/ChunkPlugin.cpp
@@ -75,6 +75,67 @@ public:
 ////////////////////////////////////////////////////////////////////////////////
 struct FChunkPluginSettings
 {
+	int32 LoadFromFile(const FString& SettingsFile)
+	{
+		if (!SettingsFile.IsEmpty())
+		{
+			UE_LOG(LogIoStoreOnDemand, Display, TEXT("Loading settings file '%s'..."), *SettingsFile);
+
+			TUniquePtr<FArchive> Ar(IFileManager::Get().CreateFileReader(*SettingsFile));
+			if (!Ar || Ar->IsError())
+			{
+				UE_LOG(LogIoStoreOnDemand, Error, TEXT("Failed to open settings file '%s'"), *SettingsFile);
+				return -1;
+			}
+
+			TSharedRef<TJsonReader<UTF8CHAR>> JsonReader = TJsonReaderFactory<UTF8CHAR>::Create(Ar.Get());
+
+			TSharedPtr<FJsonValue> JsonSettings;
+			if (!FJsonSerializer::Deserialize(*JsonReader, JsonSettings))
+			{
+				UE_LOG(LogIoStoreOnDemand, Error, TEXT("Failed to read settings file '%s'"), *SettingsFile);
+				return -1;
+			}
+
+			TSharedPtr<FJsonObject> JsonSettingsObject = JsonSettings->AsObject();
+			if (!JsonSettingsObject)
+			{
+				UE_LOG(LogIoStoreOnDemand, Error, TEXT("Bad settings file '%s'"), *SettingsFile);
+				return -1;
+			}
+
+			if (TSharedPtr<FJsonValue> PackageSetsValue = JsonSettingsObject->TryGetField(TEXT("PackageSets")))
+			{
+				TSharedPtr<FJsonObject> PackageSetsObject = PackageSetsValue->AsObject();
+				if (!PackageSetsObject)
+				{
+					UE_LOG(LogIoStoreOnDemand, Error, TEXT("Bad settings file '%s'"), *SettingsFile);
+					return -1;
+				}
+
+				for (const TPair<FString, TSharedPtr<FJsonValue>>& Pair : PackageSetsObject->Values)
+				{
+					TArray<FString> Packages;
+					if (!PackageSetsObject->TryGetStringArrayField(Pair.Key, Packages))
+					{
+						UE_LOG(LogIoStoreOnDemand, Error, TEXT("Bad settings file '%s'"), *SettingsFile);
+						return -1;
+					}
+
+					UE_LOG(LogIoStoreOnDemand, Display, TEXT("Found Package Set '%s' with %d entries"), *Pair.Key, Packages.Num());
+					for (const FString& Package : Packages)
+					{
+						UE_LOG(LogIoStoreOnDemand, Log, TEXT("->\t'%s'"), *Package);
+					}
+
+					PackageSets.Add(FString(Pair.Key), MoveTemp(Packages));
+				}
+			}
+		}
+
+		return 0;
+	}
+
 	TMap<FString, TArray<FString>> PackageSets;
 };
 
@@ -668,60 +729,9 @@ static int32 ChunkPluginCommandEntry(const FContext& Context)
 		Input.ChunkWriter = MakeUnique<FDiskChunkWriter>(OutputFolder);
 	}
 
-	if (!SettingsFile.IsEmpty())
+	if (int32 Result = Input.Settings.LoadFromFile(SettingsFile); Result < 0)
 	{
-		UE_LOG(LogIoStoreOnDemand, Display, TEXT("Loading settings file '%s'..."), *SettingsFile);
-
-		TUniquePtr<FArchive> Ar(FileMgr.CreateFileReader(*SettingsFile));
-		if (!Ar || Ar->IsError())
-		{
-			UE_LOG(LogIoStoreOnDemand, Error, TEXT("Failed to open settings file '%s'"), *SettingsFile);
-			return -1;
-		}
-
-		TSharedRef<TJsonReader<UTF8CHAR>> JsonReader = TJsonReaderFactory<UTF8CHAR>::Create(Ar.Get());
-
-		TSharedPtr<FJsonValue> JsonSettings;
-		if (!FJsonSerializer::Deserialize(*JsonReader, JsonSettings))
-		{
-			UE_LOG(LogIoStoreOnDemand, Error, TEXT("Failed to read settings file '%s'"), *SettingsFile);
-			return -1;
-		}
-
-		TSharedPtr<FJsonObject> JsonSettingsObject = JsonSettings->AsObject();
-		if (!JsonSettingsObject)
-		{
-			UE_LOG(LogIoStoreOnDemand, Error, TEXT("Bad settings file '%s'"), *SettingsFile);
-			return -1;
-		}
-
-		if (TSharedPtr<FJsonValue> PackageSets = JsonSettingsObject->TryGetField(TEXT("PackageSets")))
-		{
-			TSharedPtr<FJsonObject> PackageSetsObject = PackageSets->AsObject();
-			if (!PackageSetsObject)
-			{
-				UE_LOG(LogIoStoreOnDemand, Error, TEXT("Bad settings file '%s'"), *SettingsFile);
-				return -1;
-			}
-
-			for (const TPair<FString, TSharedPtr<FJsonValue>>& Pair : PackageSetsObject->Values)
-			{
-				TArray<FString> Packages;
-				if (!PackageSetsObject->TryGetStringArrayField(Pair.Key, Packages))
-				{
-					UE_LOG(LogIoStoreOnDemand, Error, TEXT("Bad settings file '%s'"), *SettingsFile);
-					return -1;
-				}
-
-				UE_LOG(LogIoStoreOnDemand, Display, TEXT("Found Package Set '%s' with %d entries"), *Pair.Key, Packages.Num());
-				for (const FString& Package : Packages)
-				{
-					UE_LOG(LogIoStoreOnDemand, Log, TEXT("->\t'%s'"), *Package);
-				}
-
-				Input.Settings.PackageSets.Add(FString(Pair.Key), MoveTemp(Packages));
-			}
-		}
+		return Result;
 	}
 
 	{
@@ -774,8 +784,7 @@ static int32 ChunkPluginCommandEntry(const FContext& Context)
 			continue;
 		}
 
-		int32 Result = ChunkContainer(FullPath, ContainerFlags, Input, Output);
-		if (Result < 0)
+		if (int32 Result = ChunkContainer(FullPath, ContainerFlags, Input, Output); Result < 0)
 		{
 			return Result;
 		}
@@ -881,8 +890,7 @@ static int32 ChunkPluginCommandEntry(const FContext& Context)
 
 	if (bSupportInstalledContainers && !Output.InstallOnDemandToc.Containers.IsEmpty())
 	{
-		int32 Result = WriteOnDemandToc(Output.InstallOnDemandToc, OnDemandTocName, OnDemandTocOutputFolder);
-		if (Result < 0)
+		if (int32 Result = WriteOnDemandToc(Output.InstallOnDemandToc, OnDemandTocName, OnDemandTocOutputFolder); Result < 0)
 		{
 			return Result;
 		}
@@ -891,8 +899,7 @@ static int32 ChunkPluginCommandEntry(const FContext& Context)
 	if (bSupportStreamingContainers && !Output.StreamOnDemandToc.Containers.IsEmpty())
 	{
 		// TODO: Improve how we choose the name here.
-		int32 Result = WriteOnDemandToc(Output.StreamOnDemandToc, OnDemandTocName + TEXT("ondemand"), OnDemandTocOutputFolder);
-		if (Result < 0)
+		if (int32 Result = WriteOnDemandToc(Output.StreamOnDemandToc, OnDemandTocName + TEXT("ondemand"), OnDemandTocOutputFolder); Result < 0)
 		{
 			return Result;
 		}
