CIS Fix: Mac compilation error

#rnx

- Need to include "HAL/Platform.h" before checking if the platform is windows to make sure that the macro is properly defined.

#Synced-CL 49815329

diff --git a/Engine/Source/Developer/IoStore/OnDemandUtilities/Private/Socks.cpp b/Engine/Source/Developer/IoStore/OnDemandUtilities/Private/Socks.cpp
--- a/Engine/Source/Developer/IoStore/OnDemandUtilities/Private/Socks.cpp
+++ b/Engine/Source/Developer/IoStore/OnDemandUtilities/Private/Socks.cpp
@@ -1,5 +1,7 @@
 // Copyright Epic Games, Inc. All Rights Reserved.
 
+#include "HAL/Platform.h"
+
 #if PLATFORM_WINDOWS
 
 #include "CoreGlobals.h"
