Extend the statistics recorded when running the 'ChunkPlugin' commandlet.

#rnx

- Expanded FChunkPluginStats to store the total toc and total chunk size for both IAD and IAS. Note that the values are now stored as bytes.
- We were logging the size of each new .uondemandtoc when written to disk via ::WriteOnDemandToc anyway so we can just store that value in the output stats at the same time.
- Added CPU scope when writing the stats file to disk.
- Improved the error checking when writing the stats file to disk.
- Added logging to show where the stats file was written too.

#Synced-CL 49815329

diff --git a/Engine/Source/Developer/IoStore/OnDemandUtilities/Private/ChunkPlugin.cpp b/Engine/Source/Developer/IoStore/OnDemandUtilities/Private/ChunkPlugin.cpp
--- a/Engine/Source/Developer/IoStore/OnDemandUtilities/Private/ChunkPlugin.cpp
+++ b/Engine/Source/Developer/IoStore/OnDemandUtilities/Private/ChunkPlugin.cpp
@@ -33,11 +33,19 @@ namespace UE::IoStore::Tool
 
 struct FChunkPluginStats : FJsonSerializable
 {
-	double TotalIADChunksSizeKb = 0.0;
+	int64 IADTocSize = 0;
+	int64 IADChunksSize = 0;
 
-BEGIN_JSON_SERIALIZER
-	JSON_SERIALIZE("TotalIADChunksSizeKb", TotalIADChunksSizeKb);
-END_JSON_SERIALIZER
+	int64 IASTocSize = 0;
+	int64 IASChunksSize = 0;
+
+	BEGIN_JSON_SERIALIZER
+		JSON_SERIALIZE("IADTocSize",	IADTocSize);
+		JSON_SERIALIZE("IADChunksSize",	IADChunksSize);
+
+		JSON_SERIALIZE("IASTocSize",	IASTocSize);
+		JSON_SERIALIZE("IASChunksSize",	IASChunksSize);
+	END_JSON_SERIALIZER
 };
 
 struct FS3Params
@@ -590,7 +598,14 @@ static int32 ChunkContainer(const FString& ContainerUTocPath, EOnDemandTocFlags
 		TocEntry.BlockOffset = BlockOffset;
 		TocEntry.BlockCount = BlockCount;
 
-		Output.Stats.TotalIADChunksSizeKb += (double)EncodedChunkSize / 1024.0;
+		if (EnumHasAnyFlags(TocFlags, EOnDemandTocFlags::InstallOnDemand))
+		{
+			Output.Stats.IADChunksSize += EncodedChunkSize;
+		}
+		else
+		{
+			Output.Stats.IASChunksSize += EncodedChunkSize;
+		}
 	}
 
 	if (Input.bIgnoreContainerHeader)
@@ -624,7 +639,7 @@ static int32 ChunkContainer(const FString& ContainerUTocPath, EOnDemandTocFlags
 	return 0;
 }
 
-int32 WriteOnDemandToc(FOnDemandToc& OnDemandToc, const FString& ContainerName, const FString& OutputFolder)
+int32 WriteOnDemandToc(EOnDemandTocFlags TocFlags, FOnDemandToc& OnDemandToc, const FString& ContainerName, const FString& OutputFolder, FChunkPluginStats& OutStats)
 {
 	TRACE_CPUPROFILER_EVENT_SCOPE(ChunkPluginCommand::WriteOnDemandToc);
 
@@ -637,17 +652,26 @@ int32 WriteOnDemandToc(FOnDemandToc& OnDemandToc, const FString& ContainerName,
 
 	UE::IoStore::Serialization::FOnDemandTocWriter Writer;
 	TIoStatusOr<uint64> Status = Writer.Write(OnDemandToc, TocPath);
-	if (Status.IsOk())
+	if (!Status.IsOk())
 	{
-		UE_LOG(LogIoStoreOnDemand, Display, TEXT("Saved on-demand container TOC '%s' (%.2lf KiB)"),
-			*TocPath, double(Status.ConsumeValueOrDie()) / 1024.0);
-		return 0;
+		UE_LOG(LogIoStoreOnDemand, Error, TEXT("Failed to serialize on-demand container TOC '%s', reason: %s"), *TocPath, *Status.Status().ToString());
+		return -1;
 	}
-	UE_LOG(LogIoStoreOnDemand, Error, TEXT("Failed to serialize on-demand container TOC '%s', reason: %s"),
-		*TocPath, *Status.Status().ToString());
 
+	const uint64 TocSize = Status.ConsumeValueOrDie();
+
+	if (EnumHasAnyFlags(TocFlags, EOnDemandTocFlags::InstallOnDemand))
+	{
+		OutStats.IADTocSize += TocSize;
+	}
+	else
+	{
+		OutStats.IASTocSize += TocSize;
+	}
 
-	return -1;
+	UE_LOG(LogIoStoreOnDemand, Display, TEXT("Saved on-demand container TOC '%s' (%.2lf KiB)"), *TocPath, static_cast<double>(TocSize) / 1024.0);
+	
+	return 0;
 }
 
 ////////////////////////////////////////////////////////////////////////////////
@@ -828,7 +852,7 @@ static int32 ChunkPluginCommandEntry(const FContext& Context)
 
 	if (bSupportInstalledContainers && !Output.InstallOnDemandToc.Containers.IsEmpty())
 	{
-		if (int32 Result = WriteOnDemandToc(Output.InstallOnDemandToc, OnDemandTocName, OnDemandTocOutputFolder); Result < 0)
+		if (int32 Result = WriteOnDemandToc(EOnDemandTocFlags::InstallOnDemand, Output.InstallOnDemandToc, OnDemandTocName, OnDemandTocOutputFolder, Output.Stats); Result < 0)
 		{
 			return Result;
 		}
@@ -837,7 +861,7 @@ static int32 ChunkPluginCommandEntry(const FContext& Context)
 	if (bSupportStreamingContainers && !Output.StreamOnDemandToc.Containers.IsEmpty())
 	{
 		// TODO: Improve how we choose the name here.
-		if (int32 Result = WriteOnDemandToc(Output.StreamOnDemandToc, OnDemandTocName + TEXT("ondemand"), OnDemandTocOutputFolder); Result < 0)
+		if (int32 Result = WriteOnDemandToc(EOnDemandTocFlags::StreamOnDemand, Output.StreamOnDemandToc, OnDemandTocName + TEXT("ondemand"), OnDemandTocOutputFolder, Output.Stats); Result < 0)
 		{
 			return Result;
 		}
@@ -878,7 +902,6 @@ static int32 ChunkPluginCommandEntry(const FContext& Context)
 		}
 	}
 
-
 	if (!Input.ChunkWriter->Flush())
 	{
 		UE_LOG(LogIoStoreOnDemand, Error, TEXT("Writer error: Failed to upload chunk(s)"));
@@ -887,14 +910,33 @@ static int32 ChunkPluginCommandEntry(const FContext& Context)
 
 	if (!OutputStatsJson.IsEmpty())
 	{
+		TRACE_CPUPROFILER_EVENT_SCOPE(ChunkPluginCommand::WritStatisticsFile);
+
 		TUniquePtr<FArchive> FileWriter(FileMgr.CreateFileWriter(*OutputStatsJson));
 		if (!FileWriter)
 		{
 			UE_LOG(LogIoStoreOnDemand, Display, TEXT("Failed writing stats file '%s'"), *OutputStatsJson);
 			return -1;
 		}
-		TSharedRef<TJsonWriter<UTF8CHAR>> JsonWriter = TJsonWriterFactory<UTF8CHAR>::Create(FileWriter.Get());
-		Output.Stats.ToJson(JsonWriter, false);
+
+		{
+			TSharedRef<TJsonWriter<UTF8CHAR>> JsonWriter = TJsonWriterFactory<UTF8CHAR>::Create(FileWriter.Get());
+			Output.Stats.ToJson(JsonWriter, false);
+
+			if (!JsonWriter->Close()) // Verify that the json doc was formatted correctly.
+			{
+				UE_LOG(LogIoStoreOnDemand, Display, TEXT("Failed writing stats file '%s'"), *OutputStatsJson);
+				return -1;
+			}
+		}
+
+		if (!FileWriter->Close())
+		{
+			UE_LOG(LogIoStoreOnDemand, Display, TEXT("Failed writing stats file '%s'"), *OutputStatsJson);
+			return -1;
+		}
+
+		UE_LOG(LogIoStoreOnDemand, Display, TEXT("Stats file written to '%s'"), *OutputStatsJson);
 	}
 
 	return 0;
