Added support for generating pakchunk40 uondemand toc file.

diff --git a/Engine/Source/Developer/IoStore/OnDemandUtilities/Private/Upload.cpp b/Engine/Source/Developer/IoStore/OnDemandUtilities/Private/Upload.cpp
--- a/Engine/Source/Developer/IoStore/OnDemandUtilities/Private/Upload.cpp
+++ b/Engine/Source/Developer/IoStore/OnDemandUtilities/Private/Upload.cpp
@@ -9,6 +9,7 @@
 #include "HAL/FileManager.h"
 #include "IO/IoBuffer.h"
 #include "IO/IoChunkEncoding.h"
+#include "IO/IoContainerHeader.h"
 #include "IO/IoStore.h"
 #include "IO/IoStoreOnDemand.h"
 #include "IO/Serialization/OnDemandContainerToc.h"
@@ -27,6 +28,8 @@
 namespace UE::IoStore::Tool
 {
 
+#define SHOULD_UPLOAD 1
+
 ////////////////////////////////////////////////////////////////////////////////
 struct FUploadParams
 {
@@ -46,6 +49,7 @@ struct FUploadParams
 	FString EncryptionKeyName;
 	FString HostGroupName;
 	FString GlobalContainerPath;
+	FString TagSetFile;
 	int32 MaxConcurrentUploads = 16;
 	int32 MaxTocListCount = 10000;
 	int32 MaxTocDownloadCount = 100;
@@ -135,6 +139,7 @@ static FUploadParams BuildUploadParams(const FContext& Context)
 	Ret.BuildVersion			= Context.Get<FStringView>(TEXT("-BuildVersion"),			Ret.BuildVersion);
 	Ret.TargetPlatform			= Context.Get<FStringView>(TEXT("-TargetPlatform"),			Ret.TargetPlatform);
 	Ret.HostGroupName			= Context.Get<FStringView>(TEXT("-HostGroupName"),			Ret.HostGroupName);
+	Ret.TagSetFile				= Context.Get<FStringView>(TEXT("-TagSetFile"),				Ret.TagSetFile);
 	Ret.GlobalContainerPath		= Context.Get<FStringView>(TEXT("-GlobalContainerPath"),	Ret.GlobalContainerPath);
 	Ret.bWriteTocToDisk			= Context.Get<bool>(TEXT("-WriteTocToDisk"),				Ret.bWriteTocToDisk);
 	Ret.bPerContainerTocs		= Context.Get<bool>(TEXT("-PerContainerTocs"),				Ret.bPerContainerTocs);
@@ -213,8 +218,9 @@ static FIoStatus WriteContainerFiles(FOnDemandToc& OnDemandToc, const TMap<FIoHa
 {
 	UE_LOG(LogIas, Display, TEXT("Attempting to write out %d '%s' files"), OnDemandToc.Containers.Num(), *FOnDemandToc::FileExt);
 
-	for (FOnDemandTocContainerEntry& Container : OnDemandToc.Containers)
+	for (int32 ContainerIndex = 0; ContainerIndex < OnDemandToc.Containers.Num(); ++ContainerIndex)
 	{
+		FOnDemandTocContainerEntry& Container = OnDemandToc.Containers[ContainerIndex];
 		const FString* UTocPath = UTocPaths.Find(Container.UTocHash);
 		if (UTocPath == nullptr)
 		{
@@ -228,6 +234,29 @@ static FIoStatus WriteContainerFiles(FOnDemandToc& OnDemandToc, const TMap<FIoHa
 		ContainerToc.Header = OnDemandToc.Header;
 		ContainerToc.Meta = OnDemandToc.Meta;
 		ContainerToc.Containers.Emplace(MoveTemp(Container));
+		// find all the tagsets for this container
+		for (const FOnDemandTocTagSet& TagSet : OnDemandToc.TagSets)
+		{
+			FOnDemandTocTagSet NewTagSet;
+			for (const FOnDemandTocTagSetPackageList& PackageList : TagSet.Packages)
+			{
+				if (PackageList.ContainerIndex == ContainerIndex)
+				{
+					FOnDemandTocTagSetPackageList& NewTagSetPackageList = NewTagSet.Packages[NewTagSet.Packages.AddDefaulted()];
+
+					NewTagSetPackageList.ContainerIndex = 0;
+					NewTagSetPackageList.PackageIndicies.Append(PackageList.PackageIndicies);
+				}
+			}
+
+			if (NewTagSet.Packages.Num() > 0)
+			{
+				NewTagSet.Tag = TagSet.Tag;
+				ContainerToc.TagSets.Add(NewTagSet);
+			}
+			
+		}
+		
 
 		const FString OutputPath = FPathViews::ChangeExtension(*UTocPath, FOnDemandToc::FileExt);
 
@@ -440,7 +469,9 @@ static TIoStatusOr<FUploadResult> UploadContainerFiles(
 	}
 
 	FS3Client Client(Config, Credentials);
+#if SHOULD_UPLOAD
 	FUploadQueue UploadQueue(Client, UploadParams.Bucket, UploadParams.MaxConcurrentUploads);
+#endif
 
 	if (ContainerFiles.Num() == 0)
 	{
@@ -569,6 +600,53 @@ static TIoStatusOr<FUploadResult> UploadContainerFiles(
 	uint64 TotalUploadedChunks = 0;
 	uint64 TotalUploadedBytes = 0;
 
+	TMap<FString, TArray<FString>> PackageSets;
+	if (!UploadParams.TagSetFile.IsEmpty())
+	{
+		TUniquePtr<FArchive> TagsetAr(IFileManager::Get().CreateFileReader(*UploadParams.TagSetFile));
+		if (!TagsetAr || TagsetAr->IsError())
+		{
+			UE_LOG(LogIoStoreOnDemand, Error, TEXT("Failed to open settings file '%s'"), *UploadParams.TagSetFile);
+			return FIoStatus(EIoErrorCode::ReadError, FString::Printf(TEXT("Failed to open settings file '%s'"), *UploadParams.TagSetFile));
+		}
+
+		TSharedRef<TJsonReader<UTF8CHAR>> JsonReader = TJsonReaderFactory<UTF8CHAR>::Create(TagsetAr.Get());
+
+		TSharedPtr<FJsonValue> JsonSettings;
+		if (!FJsonSerializer::Deserialize(*JsonReader, JsonSettings))
+		{
+			UE_LOG(LogIoStoreOnDemand, Error, TEXT("Failed to read settings file '%s'"), *UploadParams.TagSetFile);
+			return FIoStatus(EIoErrorCode::ReadError, FString::Printf(TEXT("Failed to read settings file '%s'"), *UploadParams.TagSetFile));
+		}
+
+		TSharedPtr<FJsonObject> PackageSetsObject = JsonSettings->AsObject();
+		if (!PackageSetsObject)
+		{
+			UE_LOG(LogIoStoreOnDemand, Error, TEXT("Bad settings file '%s'"), *UploadParams.TagSetFile);
+			return FIoStatus(EIoErrorCode::ReadError, FString::Printf(TEXT("Bad settings file '%s'"), *UploadParams.TagSetFile));
+		}
+
+
+
+		for (const TPair<FString, TSharedPtr<FJsonValue>>& Pair : PackageSetsObject->Values)
+		{
+			TArray<FString> Packages;
+			if (!PackageSetsObject->TryGetStringArrayField(Pair.Key, Packages))
+			{
+				UE_LOG(LogIas, Error, TEXT("Bad settings file '%s'"), *UploadParams.TagSetFile);
+				return FIoStatus(EIoErrorCode::ReadError, FString::Printf(TEXT("Bad settings file '%s'"), *UploadParams.TagSetFile));
+			}
+
+			UE_LOG(LogIas, Display, TEXT("Found Package Set '%s' count %d"), *Pair.Key, Packages.Num());
+			/*for (const FString& Package : Packages)
+			{
+				UE_LOG(LogIas, Display, TEXT("-> \t'%s'"), *Package);
+			}*/
+
+			PackageSets.Add(FString(Pair.Key), MoveTemp(Packages));
+		}
+	}
+
 	FOnDemandToc OnDemandToc;
 	OnDemandToc.Header.Flags = static_cast<uint32>(UploadParams.TocFlags);
 	OnDemandToc.Header.HostGroupName = UploadParams.HostGroupName;
@@ -612,11 +690,90 @@ static TIoStatusOr<FUploadResult> UploadContainerFiles(
 			return true;
 		});
 		
-		FOnDemandTocContainerEntry& ContainerEntry = OnDemandToc.Containers.AddDefaulted_GetRef();
+		const int32 ContainerIndex = OnDemandToc.Containers.AddDefaulted();
+		FOnDemandTocContainerEntry& ContainerEntry = OnDemandToc.Containers[ContainerIndex];
 		ContainerEntry.ContainerName = FPaths::GetBaseFilename(Path);
 		ContainerEntry.ContainerId = ContainerFileReader.GetContainerId();
 		ContainerEntry.ContainerFlags = static_cast<uint8>(ContainerFileReader.GetContainerFlags());
 
+
+
+
+
+
+		if (PackageSets.Num() > 0)
+		{
+			TArray<FOnDemandTocTagSet> TagSets;
+
+			const FIoChunkId HeaderChunkId = CreateContainerHeaderChunkId(ContainerFileReader.GetContainerId());
+			TIoStatusOr<FIoBuffer> Status = ContainerFileReader.Read(HeaderChunkId, FIoReadOptions());
+			if (!Status.IsOk())
+			{
+				if (Status.Status().GetErrorCode() == EIoErrorCode::NotFound)
+				{
+					UE_LOG(LogIas, Display, TEXT("No header chunk for container '%s'"), *Path);
+				}
+				else
+				{
+					UE_LOG(LogIas, Error, TEXT("Failed to read header chunk for container '%s'"), *Path);
+					return FIoStatus(EIoErrorCode::ReadError, FString::Printf(TEXT("Failed to read header chunk for container '%s'"), *Path));
+				}
+			}
+			else
+			{
+				FIoContainerHeader ContainerHeader;
+				FMemoryReaderView Ar(Status.ValueOrDie().GetView());
+				Ar << ContainerHeader;
+				Ar.Close();
+				if (Ar.IsError() || Ar.IsCriticalError())
+				{
+					UE_LOG(LogIas, Error, TEXT("Failed to deserialize header chunk for container '%s'"), *Path);
+					return FIoStatus(EIoErrorCode::ReadError, FString::Printf(TEXT("Failed to deserialize header chunk for container '%s'"), *Path));
+				}
+
+				int32 FoundPackageTagsetsCounter = 0;
+				for (const TPair<FString, TArray<FString>>& Pair : PackageSets)
+				{
+					const FString& Tag = Pair.Key;
+					if (Tag == TEXT("AICore"))
+					{
+						static int o = 0;
+						o++;
+					}
+
+					FOnDemandTocTagSet TagSet;
+					TagSet.Tag = Pair.Key;
+
+					FOnDemandTocTagSetPackageList& TagSetPackageList = TagSet.Packages[TagSet.Packages.AddDefaulted()];
+					TagSetPackageList.ContainerIndex = ContainerIndex;
+
+					bool bFoundPackageIndicies = false;
+					for (const FString& Package : Pair.Value)
+					{
+						const FPackageId PackageId = FPackageId::FromName(FName(Package));
+						int32 PackageIndex = INDEX_NONE;
+						if (ContainerHeader.PackageIds.Find(PackageId, PackageIndex))
+						{
+							++FoundPackageTagsetsCounter;
+
+							TagSetPackageList.PackageIndicies.Add(PackageIndex);
+							bFoundPackageIndicies = true;
+						}
+					}
+					UE_LOG(LogIas, Display, TEXT("Found packages %d for tag %s in container %s"), TagSetPackageList.PackageIndicies.Num(), *Tag, *Path);
+
+					if (bFoundPackageIndicies)
+					{
+						OnDemandToc.TagSets.Add(TagSet);
+					}
+				}
+
+				UE_LOG(LogIas, Display, TEXT("Found tags for %d packages in container %s"), FoundPackageTagsetsCounter, *Path);
+			}
+		}
+
+
+
 		if (EnumHasAnyFlags(ContainerFileReader.GetContainerFlags(), EIoContainerFlags::Encrypted))
 		{
 			ContainerEntry.EncryptionKeyGuid = LexToString(ContainerFileReader.GetEncryptionKeyGuid());
@@ -697,11 +854,12 @@ static TIoStatusOr<FUploadResult> UploadContainerFiles(
 					<< TEXT("/") << HashString
 					<< TEXT(".iochunk");
 
+#if SHOULD_UPLOAD
 				if (UploadQueue.Enqueue(Key, ReadResult.IoBuffer) == false)
 				{
 					return FIoStatus(EIoErrorCode::WriteError, TEXT("Failed to upload chunk"));
 				}
-
+#endif
 				TotalUploadedChunks++;
 				TotalUploadedBytes += EncodedChunkSize;
 				ContainerStats.UploadedChunkCount++;
@@ -748,7 +906,7 @@ static TIoStatusOr<FUploadResult> UploadContainerFiles(
 				Key << UploadParams.BucketPrefix.ToLower() << TEXT("/");
 			}
 			Key << LexToString(ContainerEntry.UTocHash) << TEXT(".utoc");
-
+#if SHOULD_UPLOAD
 			const FS3PutObjectResponse Response = Client.TryPutObject(
 				FS3PutObjectRequest{UploadParams.Bucket, Key.ToString(), MakeMemoryView(Buffer.GetData(), Buffer.Num())});
 
@@ -756,7 +914,7 @@ static TIoStatusOr<FUploadResult> UploadContainerFiles(
 			{
 				return FIoStatus(EIoErrorCode::WriteError, FString::Printf(TEXT("Failed to upload '%s' (%s)"), *UTocFilePath, *Response.GetErrorStatus()));
 			}
-			
+#endif
 			UTocPaths.Add(ContainerEntry.UTocHash, UTocFilePath);
 
 			UE_LOG(LogIas, Display, TEXT("Uploaded '%s'"), *UTocFilePath);	
@@ -780,12 +938,12 @@ static TIoStatusOr<FUploadResult> UploadContainerFiles(
 	{
 		return FIoStatus(EIoErrorCode::InvalidParameter, TEXT("No container file(s) marked as on demand"));
 	}
-
+#if SHOULD_UPLOAD
 	if (!UploadQueue.Flush())
 	{
 		return FIoStatus(EIoErrorCode::WriteError, TEXT("Failed to upload chunk(s)"));
 	}
-
+#endif
 	FString ServiceUrl;
 	GetTocPath(UploadParams.ServiceUrl, UploadParams.Bucket, UploadParams.BucketPrefix, ServiceUrl, OnDemandToc.Header.ChunksDirectory);
 
@@ -813,7 +971,7 @@ static TIoStatusOr<FUploadResult> UploadContainerFiles(
 		UploadResult.ServiceUrl = ServiceUrl;
 		UploadResult.CloudToc.Path = OnDemandToc.Header.ChunksDirectory / TocFilename;
 		UploadResult.CloudToc.Size = Ar.TotalSize();
-
+#if SHOULD_UPLOAD
 		const FS3PutObjectResponse Response = Client.TryPutObject(FS3PutObjectRequest{UploadParams.Bucket, Key.ToString(), Ar.GetView()});
 		if (Response.IsOk())
 		{
@@ -824,7 +982,7 @@ static TIoStatusOr<FUploadResult> UploadContainerFiles(
 			UE_LOG(LogIas, Warning, TEXT("Failed to upload TOC '%s/%s/%s' (%s)"), *Client.GetConfig().ServiceUrl, *UploadParams.Bucket, Key.ToString(), *Response.GetErrorStatus());
 			return FIoStatus(EIoErrorCode::WriteError, TEXT("Failed to upload TOC"));
 		}
-
+#endif
 		FLargeMemoryWriter FileAr;
 		if (UploadParams.bLegacyToc)
 		{
@@ -1068,6 +1226,7 @@ static FCommand UploadCommand(
 		TArgument<int32>(TEXT("-MaxConcurrentUploads"),			TEXT("Number of simultaneous uploads")),
 		TArgument<int32>(TEXT("-MaxTocListCount"),				TEXT("Maximum number of TOC file(s) to list from the bucket")),
 		TArgument<int32>(TEXT("-MaxTocDownloadCount"),			TEXT("Maximum number of TOC file(s) to download")),
+		TArgument<FStringView>(TEXT("-TagSetFile"),		TEXT("File containing tags to add to files")),
 		S3Arguments,
 	}
 );
diff --git a/Engine/Source/Developer/IoStoreUtilities/Private/IoStoreUtilities.cpp b/Engine/Source/Developer/IoStoreUtilities/Private/IoStoreUtilities.cpp
--- a/Engine/Source/Developer/IoStoreUtilities/Private/IoStoreUtilities.cpp
+++ b/Engine/Source/Developer/IoStoreUtilities/Private/IoStoreUtilities.cpp
@@ -4750,7 +4750,14 @@ static int32 CreateTarget(const FIoStoreArguments& Arguments, const FIoStoreWrit
 
 		if (bIsOptionalSegmentChunk)
 		{
-			ContainerTarget->OptionalSegmentIoStoreWriter->Append(ChunkId, WriteRequestManager.Read(TargetFile), WriteOptions);
+			if (ContainerTarget->OptionalSegmentIoStoreWriter.IsValid())
+			{ 
+				ContainerTarget->OptionalSegmentIoStoreWriter->Append(ChunkId, WriteRequestManager.Read(TargetFile), WriteOptions);
+			}
+			else
+			{
+				UE_LOG(LogIoStore, Display, TEXT("Attempted to write to optionalSegmentIoStoreWriter, but optional segment was invalid for targetfile %s"), *TargetFile.DestinationPath);
+			}
 		}
 		else
 		{
diff --git a/Engine/Source/Runtime/Core/Public/Experimental/UnifiedError/UnifiedError.h b/Engine/Source/Runtime/Core/Public/Experimental/UnifiedError/UnifiedError.h
--- a/Engine/Source/Runtime/Core/Public/Experimental/UnifiedError/UnifiedError.h
+++ b/Engine/Source/Runtime/Core/Public/Experimental/UnifiedError/UnifiedError.h
@@ -471,7 +471,11 @@ namespace UE::UnifiedError
 		template<typename T>
 		void PushErrorContext(T&& InErrorStruct, const EDetailFilter& InDetailFilter = EDetailFilter::Default)
 		{
-			PushErrorDetails(TRefCountPtr<IErrorDetails>(new TErrorDetails<T>(MoveTemp(InErrorStruct), InDetailFilter)));
+			const T* ExistingErrorContext = GetErrorContext<T>();
+			if (!ExistingErrorContext) // it's expected that only one instance of each type is inserted
+			{
+				PushErrorDetails(TRefCountPtr<IErrorDetails>(new TErrorDetails<T>(MoveTemp(InErrorStruct), InDetailFilter)));
+			}
 		}
 
 		template<typename T>
