Add support for converting all .uondemand container TOCs from a directory.

#rnx

#Synced-CL 50723010

diff --git a/Engine/Source/Developer/IoStore/OnDemandUtilities/Private/ConvertToc.cpp b/Engine/Source/Developer/IoStore/OnDemandUtilities/Private/ConvertToc.cpp
--- a/Engine/Source/Developer/IoStore/OnDemandUtilities/Private/ConvertToc.cpp
+++ b/Engine/Source/Developer/IoStore/OnDemandUtilities/Private/ConvertToc.cpp
@@ -108,6 +108,49 @@ FContainerMeta FContainerMeta::Load(const FString& Directory)
 	return OutMeta;
 }
 
+////////////////////////////////////////////////////////////////////////////////
+TArray<UE::IoStore::Serialization::V1::FOnDemandTocReader> LoadContainerTocs(const FString& PathOrDir)
+{
+	using namespace UE::IoStore::Serialization;
+	using namespace UE::IoStore::Serialization::V1;
+
+	IFileManager& Ifm = IFileManager::Get();
+
+	TArray<FString> FilePaths;
+	if (Ifm.FileExists(*PathOrDir))
+	{
+		FilePaths.Add(PathOrDir);
+	}
+	else
+	{
+		TArray<FString> Filenames;
+		Ifm.FindFiles(Filenames, *PathOrDir, FOnDemandFileExt::Toc);
+
+		for (const FString& Filename : Filenames)
+		{
+			FilePaths.Add(PathOrDir / Filename);
+		}
+	}
+
+	TArray<FOnDemandTocReader> Readers;
+	for (const FString& FilePath : FilePaths)
+	{
+		TIoStatusOr<FOnDemandTocReader> MaybeReader = FOnDemandTocReader::Read(FilePath);
+		if (MaybeReader.IsOk())
+		{
+			UE_LOG(LogIoStoreOnDemand, Display, TEXT("Loaded '%s'"), *FilePath);
+			Readers.Add(MaybeReader.ConsumeValueOrDie());
+		}
+		else
+		{
+			UE_LOG(LogIoStoreOnDemand, Error, TEXT("Failed to read on-demand TOC '%s', reason: %s"),
+				*PathOrDir, *MaybeReader.Status().ToString());
+		}
+	}
+
+	return Readers;
+}
+
 ////////////////////////////////////////////////////////////////////////////////
 static int32 ConvertTocCommandEntry(const FContext& Context)
 {
@@ -139,88 +182,90 @@ static int32 ConvertTocCommandEntry(const FContext& Context)
 		return -1;
 	}
 
-	IFileManager& Ifm = IFileManager::Get();
-
-	TIoStatusOr<FOnDemandTocReader> MaybeReader = FOnDemandTocReader::Read(FilePath);
-	if (!MaybeReader.IsOk())
+	TArray<FOnDemandTocReader> Readers = LoadContainerTocs(FilePath);
+	if (Readers.IsEmpty())
 	{
-		UE_LOG(LogIoStoreOnDemand, Error, TEXT("Failed to read on-demand TOC '%s', reason: %s"),
-			*FilePath, *MaybeReader.Status().ToString());
+		UE_LOG(LogIoStoreOnDemand, Error, TEXT("Failed to find container TOC file(s)"));
 		return -1;
 	}
 
-	FOnDemandTocReader	Reader = MaybeReader.ConsumeValueOrDie();
-	FContainerMeta		ContainerMeta = FContainerMeta::Load(FPaths::GetPath(FilePath));
+	IFileManager& Ifm = IFileManager::Get();
+	FContainerMeta ContainerMeta = FContainerMeta::Load(FPaths::GetPath(FilePath));
 
 	if (Ext == TEXT("json"))
 	{
 		FCbWriter Writer;
-		Writer.BeginObject();
+		Writer.BeginArray();
+		for (FOnDemandTocReader& Reader : Readers)
 		{
-			Writer.BeginObject(UTF8TEXTVIEW("Header"));
-			Writer.AddString(UTF8TEXTVIEW("CreationTime"), Reader.CreationTime().ToString());
-			Writer.AddString(UTF8TEXTVIEW("BuildVersion"), Reader.BuildVersion());
-			Writer.AddString(UTF8TEXTVIEW("TargetPlatform"), Reader.TargetPlatform());
-			Writer.AddString(UTF8TEXTVIEW("ChunksDirectory"), Reader.ChunksDirectory());
-			Writer.EndObject();
-
-			Writer.BeginArray(UTF8TEXTVIEW("Containers"));
-			for (const FOnDemandContainerEntry& ContainerEntry : Reader.Containers())
+			Writer.BeginObject();
 			{
-				FOnDemandTocStorage Storage;
-				TIoStatusOr<FOnDemandContainerTocView> MaybeContainerView =
-					Reader.ReadContainer(ContainerEntry, Storage, EOnDemandTocReaderOptions::None);
+				Writer.BeginObject(UTF8TEXTVIEW("Header"));
+				Writer.AddString(UTF8TEXTVIEW("CreationTime"), Reader.CreationTime().ToString());
+				Writer.AddString(UTF8TEXTVIEW("BuildVersion"), Reader.BuildVersion());
+				Writer.AddString(UTF8TEXTVIEW("TargetPlatform"), Reader.TargetPlatform());
+				Writer.AddString(UTF8TEXTVIEW("ChunksDirectory"), Reader.ChunksDirectory());
+				Writer.EndObject();
 
-				if (MaybeContainerView.IsOk() == false)
+				Writer.BeginArray(UTF8TEXTVIEW("Containers"));
+				for (const FOnDemandContainerEntry& ContainerEntry : Reader.Containers())
 				{
-					UE_LOG(LogIoStoreOnDemand, Error, TEXT("Failed to read on-demand container TOC '%s', reason: %s"),
-						*FilePath, *MaybeContainerView.Status().ToString());
-					return -1;
-				}
+					FOnDemandTocStorage Storage;
+					TIoStatusOr<FOnDemandContainerTocView> MaybeContainerView =
+						Reader.ReadContainer(ContainerEntry, Storage, EOnDemandTocReaderOptions::None);
 
-				FOnDemandContainerTocView ContainerView = MaybeContainerView.ConsumeValueOrDie();
+					if (MaybeContainerView.IsOk() == false)
+					{
+						UE_LOG(LogIoStoreOnDemand, Error, TEXT("Failed to read on-demand container TOC '%s', reason: %s"),
+							*FilePath, *MaybeContainerView.Status().ToString());
+						return -1;
+					}
 
-				Writer.BeginObject();
-				{
-					Writer.AddString(UTF8TEXTVIEW("ContainerName"), ContainerView.Header.ContainerName());
-					Writer.AddString(UTF8TEXTVIEW("ContainerId"), LexToString(ContainerView.Header.ContainerId()));
-					Writer.AddInteger(UTF8TEXTVIEW("ContainerFlags"), uint32(ContainerView.Header.ContainerFlags()));
-					Writer.AddString(UTF8TEXTVIEW("EncryptionKeyGuid"), ContainerView.Header.EncryptionKeyGuid().ToString());
-					Writer.AddInteger(UTF8TEXTVIEW("BlockSize"), ContainerView.Header.BlockSize());
-
-					TUtf8StringBuilder<512> Sb;
-					Writer.BeginArray(UTF8TEXTVIEW("ChunkEntries"));
-					for (int32 Idx = 0; const FIoChunkId& ChunkId : ContainerView.ChunkIds)
+					FOnDemandContainerTocView ContainerView = MaybeContainerView.ConsumeValueOrDie();
+
+					Writer.BeginObject();
 					{
-						FUtf8StringView		Filename = ContainerMeta.GetFilename(GetPackageChunkId(ChunkId), Sb);
-						const FChunkEntry&	ChunkEntry = ContainerView.ChunkEntries[Idx++];
-						Writer.BeginObject();
+						Writer.AddString(UTF8TEXTVIEW("ContainerName"), ContainerView.Header.ContainerName());
+						Writer.AddString(UTF8TEXTVIEW("ContainerId"), LexToString(ContainerView.Header.ContainerId()));
+						Writer.AddInteger(UTF8TEXTVIEW("ContainerFlags"), uint32(ContainerView.Header.ContainerFlags()));
+						Writer.AddString(UTF8TEXTVIEW("EncryptionKeyGuid"), ContainerView.Header.EncryptionKeyGuid().ToString());
+						Writer.AddInteger(UTF8TEXTVIEW("BlockSize"), ContainerView.Header.BlockSize());
+
+						TUtf8StringBuilder<512> Sb;
+						Writer.BeginArray(UTF8TEXTVIEW("ChunkEntries"));
+						for (int32 Idx = 0; const FIoChunkId& ChunkId : ContainerView.ChunkIds)
 						{
-							Writer.AddString(UTF8TEXTVIEW("Filename"), Filename);
-							Writer.AddString(UTF8TEXTVIEW("ChunkId"), LexToString(ChunkId));
-							Writer.AddString(UTF8TEXTVIEW("Hash"), LexToString(ChunkEntry.Hash));
-							Writer.AddInteger(UTF8TEXTVIEW("RawSize"), ChunkEntry.RawSize);
-							Writer.AddInteger(UTF8TEXTVIEW("EncodedSize"), ChunkEntry.EncodedSize);
-							Writer.AddInteger(UTF8TEXTVIEW("BlockOffset"), ChunkEntry.BlockInfo.Offset());
-							Writer.AddInteger(UTF8TEXTVIEW("BlockCount"), ChunkEntry.BlockInfo.Count());
-							if (ChunkEntry.BlockInfo.Count() == 1)
+							FUtf8StringView		Filename = ContainerMeta.GetFilename(GetPackageChunkId(ChunkId), Sb);
+							const FChunkEntry&	ChunkEntry = ContainerView.ChunkEntries[Idx++];
+							Writer.BeginObject();
 							{
-								Writer.AddInteger(UTF8TEXTVIEW("BlockSize"), ChunkEntry.BlockInfo.Size());
-								Writer.AddInteger(UTF8TEXTVIEW("BlockHash"), ChunkEntry.BlockInfo.Hash());
+								Writer.AddString(UTF8TEXTVIEW("Filename"), Filename);
+								Writer.AddString(UTF8TEXTVIEW("ChunkId"), LexToString(ChunkId));
+								Writer.AddString(UTF8TEXTVIEW("Hash"), LexToString(ChunkEntry.Hash));
+								Writer.AddInteger(UTF8TEXTVIEW("RawSize"), ChunkEntry.RawSize);
+								Writer.AddInteger(UTF8TEXTVIEW("EncodedSize"), ChunkEntry.EncodedSize);
+								Writer.AddInteger(UTF8TEXTVIEW("BlockOffset"), ChunkEntry.BlockInfo.Offset());
+								Writer.AddInteger(UTF8TEXTVIEW("BlockCount"), ChunkEntry.BlockInfo.Count());
+								if (ChunkEntry.BlockInfo.Count() == 1)
+								{
+									Writer.AddInteger(UTF8TEXTVIEW("BlockSize"), ChunkEntry.BlockInfo.Size());
+									Writer.AddInteger(UTF8TEXTVIEW("BlockHash"), ChunkEntry.BlockInfo.Hash());
+								}
 							}
+							Writer.EndObject();
 						}
-						Writer.EndObject();
+						Writer.EndArray();
 					}
-					Writer.EndArray();
+					Writer.EndObject();
 				}
-				Writer.EndObject();
+				Writer.EndArray();
 			}
-			Writer.EndArray();
+			Writer.EndObject();
 		}
-		Writer.EndObject();
+		Writer.EndArray();
 
 		TStringBuilder<4096> Json;
-		CompactBinaryToJson(Writer.Save().AsObject(), Json);
+		CompactBinaryToJson(Writer.Save().AsArray(), Json);
 
 		if (!FFileHelper::SaveStringToFile(Json.ToView(), *OutFilePath, FFileHelper::EEncodingOptions::ForceUTF8WithoutBOM))
 		{
@@ -238,41 +283,44 @@ static int32 ConvertTocCommandEntry(const FContext& Context)
 		}
 
 		Ar->Logf(TEXT("Platform, ContainerName, ChunkId, PackageId, Hash, Size, CompressedSize, Filename"));
-		for (const FOnDemandContainerEntry& ContainerEntry : Reader.Containers())
+		for (FOnDemandTocReader& Reader : Readers)
 		{
-			FOnDemandTocStorage Storage;
-			TIoStatusOr<FOnDemandContainerTocView> MaybeContainerView =
-				Reader.ReadContainer(ContainerEntry, Storage, EOnDemandTocReaderOptions::None);
-
-			if (MaybeContainerView.IsOk() == false)
+			for (const FOnDemandContainerEntry& ContainerEntry : Reader.Containers())
 			{
-				UE_LOG(LogIoStoreOnDemand, Error, TEXT("Failed to read on-demand container TOC '%s', reason: %s"),
-					*FilePath, *MaybeContainerView.Status().ToString());
-				return -1;
-			}
+				FOnDemandTocStorage Storage;
+				TIoStatusOr<FOnDemandContainerTocView> MaybeContainerView =
+					Reader.ReadContainer(ContainerEntry, Storage, EOnDemandTocReaderOptions::None);
+
+				if (MaybeContainerView.IsOk() == false)
+				{
+					UE_LOG(LogIoStoreOnDemand, Error, TEXT("Failed to read on-demand container TOC '%s', reason: %s"),
+						*FilePath, *MaybeContainerView.Status().ToString());
+					return -1;
+				}
 
-			FOnDemandContainerTocView ContainerView = MaybeContainerView.ConsumeValueOrDie();
+				FOnDemandContainerTocView ContainerView = MaybeContainerView.ConsumeValueOrDie();
 
-			TUtf8StringBuilder<512> Sb;
-			for (int32 Idx = 0; const FIoChunkId& ChunkId : ContainerView.ChunkIds)
-			{
-				FPackageId			PackageId;
-				FUtf8StringView		Filename = ContainerMeta.GetFilename(GetPackageChunkId(ChunkId, PackageId), Sb);
-				const FChunkEntry&	ChunkEntry = ContainerView.ChunkEntries[Idx++];
-				FString				ContainerFullName = FString(ContainerView.Header.ContainerName());
-
-				FString PakChunk, Platform;
-				ContainerFullName.Split(TEXT("-"), &PakChunk, &Platform, ESearchCase::CaseSensitive, ESearchDir::FromEnd);
-
-				Ar->Logf(TEXT("%s, %s, %s, %s, %s, %u, %u, %s"),
-					*Platform,
-					*PakChunk,
-					*LexToString(ChunkId),
-					*LexToString(PackageId),
-					*LexToString(ChunkEntry.Hash),
-					ChunkEntry.RawSize,
-					ChunkEntry.EncodedSize,
-					*FString(Filename));
+				TUtf8StringBuilder<512> Sb;
+				for (int32 Idx = 0; const FIoChunkId& ChunkId : ContainerView.ChunkIds)
+				{
+					FPackageId			PackageId;
+					FUtf8StringView		Filename = ContainerMeta.GetFilename(GetPackageChunkId(ChunkId, PackageId), Sb);
+					const FChunkEntry&	ChunkEntry = ContainerView.ChunkEntries[Idx++];
+					FString				ContainerFullName = FString(ContainerView.Header.ContainerName());
+
+					FString PakChunk, Platform;
+					ContainerFullName.Split(TEXT("-"), &PakChunk, &Platform, ESearchCase::CaseSensitive, ESearchDir::FromEnd);
+
+					Ar->Logf(TEXT("%s, %s, %s, %s, %s, %u, %u, %s"),
+						*Platform,
+						*PakChunk,
+						*LexToString(ChunkId),
+						*LexToString(PackageId),
+						*LexToString(ChunkEntry.Hash),
+						ChunkEntry.RawSize,
+						ChunkEntry.EncodedSize,
+						*FString(Filename));
+				}
 			}
 		}
 	}
@@ -288,7 +336,7 @@ static FCommand ConvertTocCommand(
 	TEXT("ConvertToc"),
 	TEXT("Converts an on-demand TOC file to JSON or CSV"),
 	{
-		TArgument<FStringView>(TEXT("-Path"), TEXT("Path to a .uondemandtoc file")),
+		TArgument<FStringView>(TEXT("-Path"), TEXT("Path to a .uondemandtoc file or a directory")),
 		TArgument<FStringView>(TEXT("-Out"), TEXT("Output file path"))
 	}
 );
