Add support for dumping on-demand package store info to JSON.

#rnx

#Synced-CL 51037731

diff --git a/Engine/Source/Developer/IoStore/OnDemandUtilities/Private/ConvertToc.cpp b/Engine/Source/Developer/IoStore/OnDemandUtilities/Private/ConvertToc.cpp
--- a/Engine/Source/Developer/IoStore/OnDemandUtilities/Private/ConvertToc.cpp
+++ b/Engine/Source/Developer/IoStore/OnDemandUtilities/Private/ConvertToc.cpp
@@ -4,6 +4,7 @@
 
 #include "HAL/FileManager.h"
 #include "IO/IoChunkId.h"
+#include "IO/IoContainerHeader.h"
 #include "IO/IoContainerMeta.h"
 #include "IO/IoStoreOnDemand.h"
 #include "IO/OnDemandToc.h"
@@ -11,9 +12,11 @@
 #include "IO/Serialization/OnDemandContainerToc.h"
 #include "Misc/FileHelper.h"
 #include "Misc/Paths.h"
+#include "Misc/SecureHash.h"
 #include "Misc/StringBuilder.h"
 #include "Serialization/CompactBinarySerialization.h"
 #include "Serialization/CompactBinaryWriter.h"
+#include "Serialization/LargeMemoryReader.h"
 
 namespace UE::IoStore::Tool
 {
@@ -189,6 +192,9 @@ static int32 ConvertTocCommandEntry(const FContext& Context)
 		return -1;
 	}
 
+	const bool bIncludePackageStore = Context.Get<bool>(TEXT("-IncludePackageStoreEntries"), false);
+	const bool bExcludeChunkEntries = Context.Get<bool>(TEXT("-ExcludeChunkEntries"), false);
+
 	IFileManager& Ifm = IFileManager::Get();
 	FContainerMeta ContainerMeta = FContainerMeta::Load(FPaths::GetPath(FilePath));
 
@@ -231,30 +237,84 @@ static int32 ConvertTocCommandEntry(const FContext& Context)
 						Writer.AddString(UTF8TEXTVIEW("EncryptionKeyGuid"), ContainerView.Header.EncryptionKeyGuid().ToString());
 						Writer.AddInteger(UTF8TEXTVIEW("BlockSize"), ContainerView.Header.BlockSize());
 
-						TUtf8StringBuilder<512> Sb;
-						Writer.BeginArray(UTF8TEXTVIEW("ChunkEntries"));
-						for (int32 Idx = 0; const FIoChunkId& ChunkId : ContainerView.ChunkIds)
+						if (bExcludeChunkEntries == false)
+						{
+							TUtf8StringBuilder<512> Sb;
+							Writer.BeginArray(UTF8TEXTVIEW("ChunkEntries"));
+							for (int32 Idx = 0; const FIoChunkId& ChunkId : ContainerView.ChunkIds)
+							{
+								FUtf8StringView		Filename = ContainerMeta.GetFilename(GetPackageChunkId(ChunkId), Sb);
+								const FChunkEntry&	ChunkEntry = ContainerView.ChunkEntries[Idx++];
+								Writer.BeginObject();
+								{
+									Writer.AddString(UTF8TEXTVIEW("Filename"), Filename);
+									Writer.AddString(UTF8TEXTVIEW("ChunkId"), LexToString(ChunkId));
+									Writer.AddString(UTF8TEXTVIEW("Hash"), LexToString(ChunkEntry.Hash));
+									Writer.AddInteger(UTF8TEXTVIEW("RawSize"), ChunkEntry.RawSize);
+									Writer.AddInteger(UTF8TEXTVIEW("EncodedSize"), ChunkEntry.EncodedSize);
+									Writer.AddInteger(UTF8TEXTVIEW("BlockOffset"), ChunkEntry.BlockInfo.Offset());
+									Writer.AddInteger(UTF8TEXTVIEW("BlockCount"), ChunkEntry.BlockInfo.Count());
+									if (ChunkEntry.BlockInfo.Count() == 1)
+									{
+										Writer.AddInteger(UTF8TEXTVIEW("BlockSize"), ChunkEntry.BlockInfo.Size());
+										Writer.AddInteger(UTF8TEXTVIEW("BlockHash"), ChunkEntry.BlockInfo.Hash());
+									}
+								}
+								Writer.EndObject();
+							}
+							Writer.EndArray();
+						}
+						
+						if (bIncludePackageStore)
 						{
-							FUtf8StringView		Filename = ContainerMeta.GetFilename(GetPackageChunkId(ChunkId), Sb);
-							const FChunkEntry&	ChunkEntry = ContainerView.ChunkEntries[Idx++];
-							Writer.BeginObject();
+							FLargeMemoryReader Ar(
+								ContainerView.ContainerHeaderChunk.GetData(),
+								ContainerView.ContainerHeaderChunk.GetSize());
+							
+							FIoContainerHeader ContainerHeader;
+							Ar << ContainerHeader;
+
+							if (Ar.IsError() || Ar.IsCriticalError())
 							{
-								Writer.AddString(UTF8TEXTVIEW("Filename"), Filename);
-								Writer.AddString(UTF8TEXTVIEW("ChunkId"), LexToString(ChunkId));
-								Writer.AddString(UTF8TEXTVIEW("Hash"), LexToString(ChunkEntry.Hash));
-								Writer.AddInteger(UTF8TEXTVIEW("RawSize"), ChunkEntry.RawSize);
-								Writer.AddInteger(UTF8TEXTVIEW("EncodedSize"), ChunkEntry.EncodedSize);
-								Writer.AddInteger(UTF8TEXTVIEW("BlockOffset"), ChunkEntry.BlockInfo.Offset());
-								Writer.AddInteger(UTF8TEXTVIEW("BlockCount"), ChunkEntry.BlockInfo.Count());
-								if (ChunkEntry.BlockInfo.Count() == 1)
+								UE_LOG(LogIoStoreOnDemand, Error, TEXT("Failed to serialize container header"));
+							}
+							else
+							{
+								TConstArrayView<FFilePackageStoreEntry> PackageStoreEntries = MakeArrayView(
+									reinterpret_cast<const FFilePackageStoreEntry*>(
+										ContainerHeader.StoreEntries.GetData()),
+										ContainerHeader.PackageIds.Num());
+
+								Writer.BeginArray(UTF8TEXTVIEW("PackageStoreEntries"));
+								for (int32 Index = 0; const FPackageId& PackageId : ContainerHeader.PackageIds)
 								{
-									Writer.AddInteger(UTF8TEXTVIEW("BlockSize"), ChunkEntry.BlockInfo.Size());
-									Writer.AddInteger(UTF8TEXTVIEW("BlockHash"), ChunkEntry.BlockInfo.Hash());
+									const FFilePackageStoreEntry& Entry = PackageStoreEntries[Index++];
+
+									Writer.BeginObject();
+									{
+										Writer.AddString(UTF8TEXTVIEW("PackageId"), LexToString(PackageId));
+										Writer.BeginArray(UTF8TEXTVIEW("ImportedPackageIds"));
+										{
+											for (const FPackageId& Import : Entry.ImportedPackages)
+											{
+												Writer << LexToString(Import);
+											}
+										}
+										Writer.EndArray();
+										Writer.BeginArray(UTF8TEXTVIEW("ShaderMapHashes"));
+										{
+											for (const FSHAHash& Hash : Entry.ShaderMapHashes)
+											{
+												Writer << LexToString(Hash);
+											}
+										}
+										Writer.EndArray();
+									}
+									Writer.EndObject();
 								}
+								Writer.EndArray();
 							}
-							Writer.EndObject();
 						}
-						Writer.EndArray();
 					}
 					Writer.EndObject();
 				}
@@ -337,6 +397,8 @@ static FCommand ConvertTocCommand(
 	TEXT("Converts an on-demand TOC file to JSON or CSV"),
 	{
 		TArgument<FStringView>(TEXT("-Path"), TEXT("Path to a .uondemandtoc file or a directory")),
+		TArgument<bool>(TEXT("-IncludePackageStoreEntries"), TEXT("Whether to include package store information.")),
+		TArgument<bool>(TEXT("-ExcludeChunkEntries"), TEXT("Whether to exclude chunk entry information.")),
 		TArgument<FStringView>(TEXT("-Out"), TEXT("Output file path"))
 	}
 );
