Fixed TSAN compile errors

TTSANSafeValue doesn't support chained assignment

#rnx

#Synced-CL 51332317

diff --git a/Engine/Source/Developer/IoStore/OnDemandUtilities/Private/ConvertToc.cpp b/Engine/Source/Developer/IoStore/OnDemandUtilities/Private/ConvertToc.cpp
--- a/Engine/Source/Developer/IoStore/OnDemandUtilities/Private/ConvertToc.cpp
+++ b/Engine/Source/Developer/IoStore/OnDemandUtilities/Private/ConvertToc.cpp
@@ -161,7 +161,8 @@ static int32 ConvertTocCommandEntry(const FContext& Context)
 	using namespace UE::IoStore::Serialization::V1;
 	using FChunkEntry = UE::IoStore::Serialization::V1::FOnDemandChunkEntry;
 
-	GPrintLogCategory = GPrintLogVerbosity = false;
+	GPrintLogCategory = false;
+	GPrintLogVerbosity = false;
 	GPrintLogTimes = ELogTimes::None;
 
 	FString FilePath = FString(Context.Get<FStringView>(TEXT("-Path")));
diff --git a/Engine/Source/Developer/IoStore/OnDemandUtilities/Private/Graph.cpp b/Engine/Source/Developer/IoStore/OnDemandUtilities/Private/Graph.cpp
--- a/Engine/Source/Developer/IoStore/OnDemandUtilities/Private/Graph.cpp
+++ b/Engine/Source/Developer/IoStore/OnDemandUtilities/Private/Graph.cpp
@@ -907,7 +907,8 @@ static TArray<FString> GlobContainers(const FContext& Context)
 ////////////////////////////////////////////////////////////////////////////////
 static int32 BuildGraphCommandEntry(const FContext& Context)
 {
-	GPrintLogCategory = GPrintLogVerbosity = false;
+	GPrintLogCategory = false;
+	GPrintLogVerbosity = false;
 	GPrintLogTimes = ELogTimes::None;
 
 	TArray<FString> Containers	= GraphTool::GlobContainers(Context);
@@ -943,7 +944,8 @@ static int32 BuildGraphCommandEntry(const FContext& Context)
 ////////////////////////////////////////////////////////////////////////////////
 static int32 QueryGraphCommandEntry(const FContext& Context)
 {
-	GPrintLogCategory = GPrintLogVerbosity = false;
+	GPrintLogCategory = false;
+	GPrintLogVerbosity = false;
 	GPrintLogTimes = ELogTimes::None;
 
 	FString			GraphFilename = FString(Context.Get<FStringView>(TEXT("Graph")));
diff --git a/Engine/Source/Developer/IoStore/OnDemandUtilities/Private/ResaveToc.cpp b/Engine/Source/Developer/IoStore/OnDemandUtilities/Private/ResaveToc.cpp
--- a/Engine/Source/Developer/IoStore/OnDemandUtilities/Private/ResaveToc.cpp
+++ b/Engine/Source/Developer/IoStore/OnDemandUtilities/Private/ResaveToc.cpp
@@ -14,7 +14,8 @@ namespace UE::IoStore::Tool
 ////////////////////////////////////////////////////////////////////////////////
 static int32 ResaveTocEntry(const FContext& Context)
 {
-	GPrintLogCategory = GPrintLogVerbosity = false;
+	GPrintLogCategory = false;
+	GPrintLogVerbosity = false;
 	GPrintLogTimes = ELogTimes::None;
 
 	IFileManager& Ifm = IFileManager::Get();
