Updated IAS tool graph command to parse package name or ID.

#rnx

#Synced-CL 51343971

diff --git a/Engine/Source/Developer/IoStore/OnDemandUtilities/Private/Graph.cpp b/Engine/Source/Developer/IoStore/OnDemandUtilities/Private/Graph.cpp
--- a/Engine/Source/Developer/IoStore/OnDemandUtilities/Private/Graph.cpp
+++ b/Engine/Source/Developer/IoStore/OnDemandUtilities/Private/Graph.cpp
@@ -955,18 +955,16 @@ static int32 QueryGraphCommandEntry(const FContext& Context)
 	bool		bPrintDeps	= false;
 	FPackageId	PackageId	= FPackageId();
 
-	if (FStringView HexView = Context.Get<FStringView>(TEXT("-Traverse")); HexView.IsEmpty() == false)
+	if (FStringView PackageNameOrId = Context.Get<FStringView>(TEXT("-Traverse")); PackageNameOrId.IsEmpty() == false)
 	{
-		uint64 Value	= FParse::HexNumber64(FString(HexView));
-		PackageId		= FPackageId::FromValue(Value);
-		bTraverse		= true;
+		LexFromString(PackageId, PackageNameOrId);
+		bTraverse = true;
 	}
 
-	if (FStringView HexView = Context.Get<FStringView>(TEXT("-Deps")); HexView.IsEmpty() == false)
+	if (FStringView PackageNameOrId = Context.Get<FStringView>(TEXT("-Deps")); PackageNameOrId.IsEmpty() == false)
 	{
-		uint64 Value	= FParse::HexNumber64(FString(HexView));
-		PackageId		= FPackageId::FromValue(Value);
-		bPrintDeps		= true;
+		LexFromString(PackageId, PackageNameOrId);
+		bPrintDeps	= true;
 	}
 
 	if (PackageId.IsValid() == false)
