h2 compatible case

#rnx

#Synced-CL 51451766

diff --git a/Engine/Source/Developer/IoStore/OnDemandUtilities/Private/Purl.cpp b/Engine/Source/Developer/IoStore/OnDemandUtilities/Private/Purl.cpp
--- a/Engine/Source/Developer/IoStore/OnDemandUtilities/Private/Purl.cpp
+++ b/Engine/Source/Developer/IoStore/OnDemandUtilities/Private/Purl.cpp
@@ -140,7 +140,7 @@ static int32 PurlCommandEntry(const FContext& Context)
 	}
 
 	FRequest Request = Loop.Request(AnsiMethod, AnsiPath, Pool, &RequestParams);
-	Request.Header("User-Agent", "purl");
+	Request.Header("user-agent", "purl");
 	Loop.Send(MoveTemp(Request), Sink);
 
 	Loop.SetFailTimeout(4000);
