On-demand chunk partition(s)

Replace the per-chunk file model with a partition model where chunks are packed into larger blobs (.upart). To keep TOC memory usage down, the per-chunk hash is
reduced from 20 to 12 bytes. This is a safe reduction since the install cache already uses 12-byte content hashes for chunk addressing.

Key changes:
- Add V2 serialization format with FOnDemandPartitionEntry and updated FOnDemandChunkEntry (partition index/offset instead of chunk hash)
- Replace FOnDemandToc struct with FOnDemandTocWriter builder API; containers/partitions are built incrementally via Begin/Add/End calls
- Update Upload.cpp to pack chunks into partitions and upload partition TOC manifests (.uparttoc) for deduplication on incremental uploads
- Update ChunkPlugin, ContentInstaller, InstallCache, and HTTP dispatcher to use new partition-aware APIs
- Remove ResaveToc.cpp (superseded by new TOC writer)
- Added new file extensions .iopart and .ioparttoc for chunk partition and cloud TOC

Note: Custom partitioning layout/ordering and the IasTool commands Download and ListTocs will be fixed
in a follow up change(s).

#rnx

#Synced-CL 52723932

diff --git a/Engine/Source/Developer/IoStore/OnDemandUtilities/Private/ChunkPlugin.cpp b/Engine/Source/Developer/IoStore/OnDemandUtilities/Private/ChunkPlugin.cpp
--- a/Engine/Source/Developer/IoStore/OnDemandUtilities/Private/ChunkPlugin.cpp
+++ b/Engine/Source/Developer/IoStore/OnDemandUtilities/Private/ChunkPlugin.cpp
@@ -2,6 +2,7 @@
 
 #include "Command.h"
 #include "Common.h"
+#include "Upload.h"
 
 #include "Algo/Find.h"
 #include "Containers/StringView.h"
@@ -12,7 +13,6 @@
 #include "IO/IoStore.h"
 #include "IO/IoStoreOnDemand.h"
 #include "IO/Serialization/OnDemandContainerToc.h"
-#include "Misc/Compression.h"
 #include "Misc/DateTime.h"
 #include "Misc/FileHelper.h"
 #include "Misc/KeyChainUtilities.h"
@@ -31,6 +31,9 @@
 namespace UE::IoStore::Tool
 {
 
+using FOnDemandTocWriter			= UE::IoStore::Serialization::FOnDemandTocWriter;
+using EOnDemandContainerEntryFlags	= UE::IoStore::Serialization::EOnDemandContainerEntryFlags;
+
 struct FChunkPluginStats : FJsonSerializable
 {
 	int64 IADTocSize = 0;
@@ -187,7 +190,7 @@ public:
 			<< RelativeDirectory
 			<< TEXT("/") << HashString.Left(2)
 			<< TEXT("/") << HashString
-			<< TEXT(".iochunk");
+			<< UE::IoStore::Serialization::FOnDemandFileExt::Partition;
 
 		if (UploadQueue.Enqueue(Key, Chunk) == false)
 		{
@@ -245,7 +248,7 @@ public:
 				<< TEXT("'");
 		}
 
-		Sb << TEXT("/") << HashString << TEXT(".iochunk");
+		Sb << TEXT("/") << HashString << UE::IoStore::Serialization::FOnDemandFileExt::Partition;
 
 		TUniquePtr<FArchive> Ar(IFileManager::Get().CreateFileWriter(Sb.ToString()));
 
@@ -274,47 +277,6 @@ private:
 	FString OutputFolder;
 };
 
-////////////////////////////////////////////////////////////////////////////////
-
-static bool ShouldValidateChunk(FName CompressionFormat, uint64 ChunkRawSize)
-{
-	static TMap<FName, bool> ValidatedFormats;
-
-	if (!CompressionFormat.IsNone())
-	{
-		if (bool* bCachedResult = ValidatedFormats.Find(CompressionFormat))
-		{
-			if (*bCachedResult == false)
-			{
-				return false;
-			}
-		}
-		else
-		{
-			// Note that if IsFormatValid returns false it will log a message about it. This is one of the
-			// reasons we must make sure to only check once per format to avoid  log spam. The other reason
-			// is that the check is fairly slow.
-			bool bValid = FCompression::IsFormatValid(CompressionFormat);
-			ValidatedFormats.Add(CompressionFormat, bValid);
-
-			if (bValid == false)
-			{
-				UE_LOGF(LogIoStoreOnDemand, Display, "Skipping validation tests for chunks with the compression format '%ls'", *CompressionFormat.ToString())
-				return false;
-			}
-		}
-	}
-
-	const uint64 MaxChunkSize = 32 * 1024 * 1024;	// We do not want to validate any chunk that will decompress to more than 32 MiB
-													// as we don't want to use too much memory.
-	if (ChunkRawSize >= MaxChunkSize)
-	{
-		return false;
-	}
-
-	return true;
-}
-
 struct FChunkParams
 {
 	FChunkPluginSettings Settings;
@@ -323,8 +285,8 @@ struct FChunkParams
 	TUniquePtr<FChunkWriterInterface> ChunkWriter;
 	FString ChunksRelativeFolder;
 
-	bool bIgnoreContainerHeader;
-	bool bValidateChunks;
+	int32 MaxPartitionSize = 0;
+	bool bValidateChunks = false;
 };
 
 struct FChunkOutput
@@ -333,30 +295,18 @@ struct FChunkOutput
 	{
 		const TCHAR* ChunksDirectory = TEXT(""); // TODO: Should we be setting this to something?
 
-		InstallOnDemandToc.Header.Flags = static_cast<uint32>(EOnDemandTocFlags::InstallOnDemand);
-		InstallOnDemandToc.Header.HostGroupName = HostGroupName;
-		InstallOnDemandToc.Header.ChunksDirectory = ChunksDirectory;
-		InstallOnDemandToc.Meta.TargetPlatform = Platform;
-		InstallOnDemandToc.Meta.BuildVersion = BuildVersion;
-
-		StreamOnDemandToc.Header.Flags = static_cast<uint32>(EOnDemandTocFlags::StreamOnDemand);
-		StreamOnDemandToc.Header.HostGroupName = HostGroupName;
-		StreamOnDemandToc.Header.ChunksDirectory = ChunksDirectory;
-		StreamOnDemandToc.Meta.TargetPlatform = Platform;
-		StreamOnDemandToc.Meta.BuildVersion = BuildVersion;
-	}
-
-	void Reserve(int32 NumPotentialContainers, int32 NumPackageSets)
-	{
-		InstallOnDemandToc.Containers.Reserve(NumPotentialContainers);
-		StreamOnDemandToc.Containers.Reserve(NumPotentialContainers);
+		InstallOnDemandToc.SetMetadata(BuildVersion, Platform);
+		InstallOnDemandToc.SetChunksDirectory(ChunksDirectory);
+		InstallOnDemandToc.SetHostGroup(HostGroupName);
 
-		FoundTagSets.Reserve(NumPackageSets);
+		StreamOnDemandToc.SetMetadata(BuildVersion, Platform);
+		StreamOnDemandToc.SetChunksDirectory(ChunksDirectory);
+		StreamOnDemandToc.SetHostGroup(HostGroupName);
 	}
 
-	FOnDemandToc& GetToc(EOnDemandTocFlags Flags)
+	FOnDemandTocWriter& GetToc(EOnDemandContainerEntryFlags Flags)
 	{
-		if (EnumHasAnyFlags(Flags, EOnDemandTocFlags::InstallOnDemand))
+		if (EnumHasAnyFlags(Flags, EOnDemandContainerEntryFlags::InstallOnDemand))
 		{
 			return InstallOnDemandToc;
 		}
@@ -366,280 +316,96 @@ struct FChunkOutput
 		}
 	}
 
-	void SetTocsTimestamp(int64 Timestamp)
-	{
-		InstallOnDemandToc.Meta.EpochTimestamp = Timestamp;
-		StreamOnDemandToc.Meta.EpochTimestamp = Timestamp;
-	}
+	FOnDemandTocWriter InstallOnDemandToc;
+	FOnDemandTocWriter StreamOnDemandToc;
 
-	FOnDemandToc InstallOnDemandToc;
-	FOnDemandToc StreamOnDemandToc;
-
-	/* Tracks the packages found for each tag per container. Is only used for InstallOnDemandToc. */
-	TMap<FString, TArray<FOnDemandTocTagSetPackageList>> FoundTagSets;
 	/* The paths of the container files (.ucas, .utoc etc) that were chunked. */
 	TArray<FString> ContainerFilePaths;
 	/* Statistics about the chunking process. */
 	FChunkPluginStats Stats;
 };
 
-static int32 ChunkContainer(const FString& ContainerUTocPath, EOnDemandTocFlags TocFlags, const FChunkParams& Input, FChunkOutput& Output)
+static int32 ChunkContainer(const FString& ContainerUTocPath, EOnDemandContainerEntryFlags ContainerFlags, const FChunkParams& Input, FChunkOutput& Output)
 {
+	using namespace UE::IoStore::Tool::Upload;
+
 	TRACE_CPUPROFILER_EVENT_SCOPE(ChunkPluginCommand::ChunkContainer);
 
 	check(Input.ChunkWriter.IsValid());
 
-	FIoStoreReader ContainerFileReader;
-	{
-		FIoStatus Status = ContainerFileReader.Initialize(*FPaths::ChangeExtension(ContainerUTocPath, TEXT("")), Input.EncryptionKeys);
-		if (!Status.IsOk())
-		{
-			UE_LOGF(LogIoStoreOnDemand, Error, "Failed to open container '%ls' for reading due to: %ls", *ContainerUTocPath, *Status.ToString());
-			return -1;
-		}
-	}
-
-	UE_LOGF(LogIoStoreOnDemand, Display, "Processing container '%ls'", *ContainerUTocPath);
-
-	FOnDemandToc& OnDemandToc = Output.GetToc(TocFlags);
-
-	const uint32 BlockSize = ContainerFileReader.GetCompressionBlockSize();
-	if (OnDemandToc.Header.BlockSize == 0)
-	{
-		OnDemandToc.Header.BlockSize = ContainerFileReader.GetCompressionBlockSize();
+	TConstArrayView<FString> FileNames(&ContainerUTocPath, 1);
+	TPagedArray<FContainerData> Containers;
+	const bool bIgnoreContainerFlags = true;
+	if (FIoStatus Status = LoadContainers(
+		FileNames,
+		Input.EncryptionKeys,
+		Containers,
+		bIgnoreContainerFlags); !Status.IsOk())
+	{
+		UE_LOGF(LogIoStoreOnDemand, Error, "Failed to open container '%ls' for reading due to: %ls",
+			*ContainerUTocPath, *Status.ToString());
+		return -1;
 	}
 
-	check(OnDemandToc.Header.BlockSize == ContainerFileReader.GetCompressionBlockSize());
-
-	const int32 ContainerIndex = OnDemandToc.Containers.AddDefaulted();
-	FOnDemandTocContainerEntry& ContainerEntry = OnDemandToc.Containers[ContainerIndex];
-	ContainerEntry.ContainerId = ContainerFileReader.GetContainerId();
-	ContainerEntry.ContainerName = FPaths::GetBaseFilename(ContainerUTocPath);
-
-	const bool bIsContainerEncrypted = EnumHasAnyFlags(ContainerFileReader.GetContainerFlags(), EIoContainerFlags::Encrypted);
-	if (bIsContainerEncrypted)
+	if (Containers.IsEmpty())
 	{
-		ContainerEntry.EncryptionKeyGuid = LexToString(ContainerFileReader.GetEncryptionKeyGuid());
-		ContainerEntry.ContainerFlags |= static_cast<uint8>(EIoContainerFlags::Encrypted);
+		UE_LOGF(LogIoStoreOnDemand, Error, "Failed to load container '%ls'", *ContainerUTocPath);
+		return -1;
 	}
 
-	if (EnumHasAnyFlags(TocFlags, EOnDemandTocFlags::InstallOnDemand) && Input.Settings.PackageSets.Num() > 0)
+	const FTagSets& TagSets = Input.Settings.PackageSets;
+	if (EnumHasAnyFlags(ContainerFlags, EOnDemandContainerEntryFlags::InstallOnDemand) && !TagSets.IsEmpty())
 	{
-		TArray<FOnDemandTocTagSet> TagSets;
-
-		const FIoChunkId HeaderChunkId = CreateContainerHeaderChunkId(ContainerFileReader.GetContainerId());
-		TIoStatusOr<FIoBuffer> Status = ContainerFileReader.Read(HeaderChunkId, FIoReadOptions());
-		if (!Status.IsOk())
-		{
-			if (Status.Status().GetErrorCode() == EIoErrorCode::NotFound)
-			{
-				UE_LOGF(LogIoStoreOnDemand, Display, "No header chunk for container '%ls'", *ContainerUTocPath);
-			}
-			else
-			{
-				UE_LOGF(LogIoStoreOnDemand, Error, "Failed to read header chunk for container '%ls'", *ContainerUTocPath);
-				return -1;
-			}
-		}
-		else
+		// Can only create tag sets for containers with package information
+		if (Containers[0].Header.GetSize() > 0)
 		{
-			FIoContainerHeader ContainerHeader;
-			FMemoryReaderView Ar(Status.ValueOrDie().GetView());
-			Ar << ContainerHeader;
-			Ar.Close();
-			if (Ar.IsError() || Ar.IsCriticalError())
+			if (FIoStatus Status = CreateTagSets(TagSets, Containers[0]); !Status.IsOk())
 			{
-				UE_LOGF(LogIoStoreOnDemand, Error, "Failed to deserialize header chunk for container '%ls'", *ContainerUTocPath);
+				UE_LOGF(LogIoStoreOnDemand, Error, "Failed to create tag set(s) for container '%ls' due to: %ls",
+						*ContainerUTocPath, *Status.ToString());
 				return -1;
 			}
-
-			for (const TPair<FString, TArray<FString>>& Pair : Input.Settings.PackageSets)
-			{
-				const FString& Tag = Pair.Key;
-
-				FOnDemandTocTagSetPackageList TagSetPackageList;
-				TagSetPackageList.ContainerIndex = ContainerIndex;
-
-				for (const FString& Package : Pair.Value)
-				{
-					const FPackageId PackageId = FPackageId::FromName(FName(Package));
-					int32 PackageIndex = INDEX_NONE;
-					if (ContainerHeader.PackageIds.Find(PackageId, PackageIndex))
-					{
-						UE_LOGF(LogIoStoreOnDemand, Log, "Found package %ls for tag %ls in container %ls", *Package, *Tag, *ContainerUTocPath);
-						TagSetPackageList.PackageIndicies.Add(PackageIndex);
-					}
-				}
-
-				if (TagSetPackageList.PackageIndicies.Num() > 0)
-				{
-					Output.FoundTagSets.FindOrAdd(Tag).Add(MoveTemp(TagSetPackageList));
-				}
-			}
 		}
 	}
 
-	TArray<FIoStoreTocChunkInfo> ChunkInfos;
-	ContainerFileReader.EnumerateChunks([&ChunkInfos](FIoStoreTocChunkInfo&& Info)
-		{
-			ChunkInfos.Emplace(MoveTemp(Info));
-			return true;
-		});
+	UE_LOGF(LogIoStoreOnDemand, Display, "Processing container '%ls'", *ContainerUTocPath);
 
-	UE_LOGF(LogIoStoreOnDemand, Display, "Serializing %d chunks", ChunkInfos.Num());
+	FOnDemandTocWriter& TocWriter = Output.GetToc(ContainerFlags);
 
-	for (const FIoStoreTocChunkInfo& ChunkInfo : ChunkInfos)
+	auto WriteChunk = [&Input](FIoBuffer Chunk, const FIoHash& Hash) -> FIoStatus
 	{
-		const bool bDecrypt = false; // We do not want to decrypt the results. We will upload encrypted chunks and decrypt on the client
-		TIoStatusOr<FIoStoreCompressedReadResult> Status = ContainerFileReader.ReadCompressed(ChunkInfo.Id, FIoReadOptions(), bDecrypt);
-		if (!Status.IsOk())
-		{
-			UE_LOGF(LogIoStoreOnDemand, Error, "Failed to read container chunk, Container='%ls', Reason='%ls'",
-				*ContainerUTocPath, *Status.Status().ToString());
-			return -1;
-		}
-
-		FIoStoreCompressedReadResult ReadResult = Status.ConsumeValueOrDie();
-
-		const uint32 BlockOffset = ContainerEntry.BlockSizes.Num();
-		const uint32 BlockCount = ReadResult.Blocks.Num();
-		auto HashBuffer = [](const FIoStoreCompressedReadResult& ReadResult)
-		{
-			TRACE_CPUPROFILER_EVENT_SCOPE(HashBuffer);
-			return FIoHash::HashBuffer(ReadResult.IoBuffer.GetView());
-		};
-		const FIoHash ChunkHash = HashBuffer(ReadResult);
-		const FString HashString = LexToString(ChunkHash);
-
-		FMemoryView EncodedBlocks = ReadResult.IoBuffer.GetView();
-		uint64 RawChunkSize = 0;
-		uint64 EncodedChunkSize = 0;
-		FName CompressionFormat = NAME_None;
-
-		{
-			TRACE_CPUPROFILER_EVENT_SCOPE(HashBlocks);
-			for (const FIoStoreCompressedBlockInfo& BlockInfo : ReadResult.Blocks)
-			{
-				check(Align(BlockInfo.CompressedSize, FAES::AESBlockSize) == BlockInfo.AlignedSize);
-				const uint64 EncodedBlockSize = BlockInfo.AlignedSize;
-				ContainerEntry.BlockSizes.Add(uint32(BlockInfo.CompressedSize));
-
-				FMemoryView EncodedBlock = EncodedBlocks.Left(EncodedBlockSize);
-				EncodedBlocks += EncodedBlock.GetSize();
-
-				check(((char*)EncodedBlock.GetData() + EncodedBlock.GetSize()) <= ((char*)ReadResult.IoBuffer.GetView().GetData() + ReadResult.IoBuffer.GetView().GetSize()));
-				FIoBlockHash Hash = FIoChunkEncoding::HashBlock(EncodedBlock);
-				ContainerEntry.BlockHashes.Add(Hash);
-
-				EncodedChunkSize += EncodedBlockSize;
-				RawChunkSize += BlockInfo.UncompressedSize;
-
-				if (BlockInfo.CompressionMethod != NAME_None)
-				{
-					CompressionFormat = BlockInfo.CompressionMethod;
-					if (OnDemandToc.Header.CompressionFormat.IsEmpty())
-					{
-						OnDemandToc.Header.CompressionFormat = BlockInfo.CompressionMethod.ToString();
-					}
-				}
-			}
-		}
-
-		if (EncodedChunkSize != ReadResult.IoBuffer.GetSize())
-		{
-			UE_LOGF(LogIoStoreOnDemand, Error, "Chunk size mismatch, Container='%ls', ChunkId='%ls'",
-				*ContainerUTocPath, *LexToString(ChunkInfo.Id));
-			return -1;
-		}
-
-		// Additional validation to make sure that the chunks can be decrypted and decoded
-		if (Input.bValidateChunks && ShouldValidateChunk(CompressionFormat, RawChunkSize))
-		{
-			TRACE_CPUPROFILER_EVENT_SCOPE(ValidateChunk);
-
-			FIoChunkDecodingParams DecodingParams;
-			DecodingParams.CompressionFormat = CompressionFormat;
-			DecodingParams.TotalRawSize = RawChunkSize;
-			DecodingParams.EncodedBlockSize = TConstArrayView<uint32>(ContainerEntry.BlockSizes.GetData() + BlockOffset, BlockCount);
-			DecodingParams.BlockHash = TConstArrayView<FIoBlockHash>(ContainerEntry.BlockHashes.GetData() + BlockOffset, BlockCount);
-
-			if (bIsContainerEncrypted)
-			{
-				const FAES::FAESKey* EncryptionKey = Input.EncryptionKeys.Find(ContainerFileReader.GetEncryptionKeyGuid());
-				if (EncryptionKey == nullptr)
-				{
-					UE_LOGF(LogIoStoreOnDemand, Error, "Failed to find encryption key '%ls', Container='%ls", *LexToString(ContainerFileReader.GetEncryptionKeyGuid()), *ContainerUTocPath);
-					return -1;
-				}
-				DecodingParams.EncryptionKey = FMemoryView(EncryptionKey->Key, EncryptionKey->KeySize);
-			}
-
-			FIoBuffer OutputBuffer(RawChunkSize);
-			bool bDecoded = FIoChunkEncoding::Decode(DecodingParams, ReadResult.IoBuffer.GetView(), OutputBuffer.GetMutableView(), EIoDecodeFlags::None);
-			if (!bDecoded)
-			{
-				UE_LOGF(LogIoStoreOnDemand, Error, "Failed to decode chunk from disk, Container='%ls', ChunkId='%ls'", *ContainerUTocPath, *LexToString(ChunkInfo.Id));
-				return -1;
-			}
-		}
+		return Input.ChunkWriter->WriteChunk(Input.ChunksRelativeFolder, Chunk, Hash);
+	};
 
-		FIoStatus WriteStatus = Input.ChunkWriter->WriteChunk(Input.ChunksRelativeFolder, ReadResult.IoBuffer, ChunkHash);
-		if (WriteStatus.IsOk() == false)
-		{
-			UE_LOGF(LogIoStoreOnDemand, Error, "%ls", *WriteStatus.ToString());
-			return -1;
-		}
-
-		FOnDemandTocEntry& TocEntry = ContainerEntry.Entries.AddDefaulted_GetRef();
-		TocEntry.ChunkId = ChunkInfo.Id;
-		TocEntry.Hash = ChunkHash;
-		TocEntry.RawSize = RawChunkSize;
-		TocEntry.EncodedSize = EncodedChunkSize;
-		TocEntry.BlockOffset = BlockOffset;
-		TocEntry.BlockCount = BlockCount;
-
-		if (EnumHasAnyFlags(TocFlags, EOnDemandTocFlags::InstallOnDemand))
-		{
-			Output.Stats.IADChunksSize += EncodedChunkSize;
-		}
-		else
-		{
-			Output.Stats.IASChunksSize += EncodedChunkSize;
-		}
-	}
-
-	if (Input.bIgnoreContainerHeader)
+	FChunkValidationSettings ValidationSettings
 	{
-		UE_LOGF(LogIoStoreOnDemand, Display, "Ignoring container header for '%ls'", *ContainerUTocPath)
-	}
-	else
+		.bValidate = Input.bValidateChunks,
+		.MaxRawChunkSize = 32 << 20
+	};
+
+	TSet<FIoHash> ExistingChunks;
+	if (FIoStatus Status = UploadContainer(
+		Containers[0],
+		Input.EncryptionKeys,
+		Input.MaxPartitionSize,
+		ContainerFlags,
+		ExistingChunks,
+		MoveTemp(WriteChunk),
+		TocWriter,
+		ValidationSettings); !Status.IsOk())
 	{
-		const bool bDecrypt = false;
-		const FIoChunkId ChunkId = CreateContainerHeaderChunkId(ContainerFileReader.GetContainerId());
-		TIoStatusOr<FIoStoreCompressedReadResult> ReadResult = ContainerFileReader.ReadCompressed(ChunkId, FIoReadOptions(), bDecrypt);
-		if (ReadResult.IsOk())
-		{
-			const FIoBuffer& Chunk = ReadResult.ValueOrDie().IoBuffer;
-			ContainerEntry.Header = MakeArrayView<const uint8>(Chunk.GetData(), static_cast<int32>(Chunk.GetSize()));
-		}
-		else
-		{
-			const EIoErrorCode Err = ReadResult.Status().GetErrorCode();
-			if (Err != EIoErrorCode::UnknownChunkID && Err != EIoErrorCode::NotFound)
-			{
-				UE_LOGF(LogIoStoreOnDemand, Error, "Failed to read container header, reason '%ls'", *ReadResult.Status().ToString())
-					return -1;
-			}
-		}
+		UE_LOGF(LogIoStoreOnDemand, Error, "Failed to process container '%ls' due to: %ls",
+			*ContainerUTocPath, *Status.ToString());
+		return -1;
 	}
 
 	Output.ContainerFilePaths.Add(ContainerUTocPath);						// .utoc
-	ContainerFileReader.GetContainerFilePaths(Output.ContainerFilePaths);	// .ucas
+	Containers[0].Reader->GetContainerFilePaths(Output.ContainerFilePaths);	// .ucas
 
 	return 0;
 }
 
-int32 WriteOnDemandToc(EOnDemandTocFlags TocFlags, FOnDemandToc& OnDemandToc, const FString& ContainerName, const FString& OutputFolder, FChunkPluginStats& OutStats)
+int32 WriteOnDemandToc(EOnDemandContainerEntryFlags TocFlags, FOnDemandTocWriter& OnDemandToc, const FString& ContainerName, const FString& OutputFolder, FChunkPluginStats& OutStats)
 {
 	TRACE_CPUPROFILER_EVENT_SCOPE(ChunkPluginCommand::WriteOnDemandToc);
 
@@ -650,8 +416,7 @@ int32 WriteOnDemandToc(EOnDemandTocFlags TocFlags, FOnDemandToc& OnDemandToc, co
 
 	const FString TocPath = OutputFolder / Filename;
 
-	UE::IoStore::Serialization::FOnDemandTocWriter Writer;
-	TIoStatusOr<uint64> Status = Writer.Write(OnDemandToc, TocPath);
+	TIoStatusOr<uint64> Status = OnDemandToc.Write(TocPath);
 	if (!Status.IsOk())
 	{
 		UE_LOGF(LogIoStoreOnDemand, Error, "Failed to serialize on-demand container TOC '%ls', reason: %ls", *TocPath, *Status.Status().ToString());
@@ -660,7 +425,7 @@ int32 WriteOnDemandToc(EOnDemandTocFlags TocFlags, FOnDemandToc& OnDemandToc, co
 
 	const uint64 TocSize = Status.ConsumeValueOrDie();
 
-	if (EnumHasAnyFlags(TocFlags, EOnDemandTocFlags::InstallOnDemand))
+	if (EnumHasAnyFlags(TocFlags, EOnDemandContainerEntryFlags::InstallOnDemand))
 	{
 		OutStats.IADTocSize += TocSize;
 	}
@@ -690,8 +455,8 @@ static int32 ChunkPluginCommandEntry(const FContext& Context)
 	FString SettingsFile				= FString(Context.Get<FStringView>(TEXT("-SettingsFile"), FString()));
 	FString OutputStatsJson				= FString(Context.Get<FStringView>(TEXT("-OutputStatsJson"), FString()));
 	const FString HostGroupName			= FString(Context.Get<FStringView>(TEXT("-HostGroupName"), FString()));
+	int32 PartitionSize					= Context.Get<int32>(TEXT("-MaxPartitionSize"), Upload::FUploadSettings::DefaultPartitionSize);
 
-	const bool bIgnoreContainerHeader	= Context.Get<bool>(TEXT("-IgnoreContainerHeader"), false);
 	const bool bSupportInstalledContainers = !Context.Get<bool>(TEXT("-SkipInstallOnDemand"), false);
 	const bool bSupportStreamingContainers = !Context.Get<bool>(TEXT("-SkipStreamOnDemand"), false);
 
@@ -725,13 +490,12 @@ static int32 ChunkPluginCommandEntry(const FContext& Context)
 	UE_LOGF(LogIoStoreOnDemand, Display, "\tOutputStatsJson: %ls", *OutputStatsJson);
 	UE_LOGF(LogIoStoreOnDemand, Display, "\tIncludeSigPak: %ls", bIncludeSigPak ? TEXT("true") : TEXT("false"));
 	UE_LOGF(LogIoStoreOnDemand, Display, "\tDeleteContainerFiles: %ls", bDeleteContainerFiles ? TEXT("true") : TEXT("false"));
-	UE_LOGF(LogIoStoreOnDemand, Display, "\tIgnoreContainerHeader: %ls", bIgnoreContainerHeader ? TEXT("true") : TEXT("false"));
 	UE_LOGF(LogIoStoreOnDemand, Display, "\tValidateChunks: %ls", bValidateChunks ? TEXT("true") : TEXT("false"));
 
 	FChunkParams Input;
 	Input.ChunksRelativeFolder = ChunksRelativeFolder;
-	Input.bIgnoreContainerHeader = bIgnoreContainerHeader;
 	Input.bValidateChunks = bValidateChunks;
+	Input.MaxPartitionSize = FMath::Clamp(PartitionSize, 0, Upload::FUploadSettings::MaxPartitionSize);
 
 	IFileManager& FileMgr = IFileManager::Get();
 
@@ -774,7 +538,6 @@ static int32 ChunkPluginCommandEntry(const FContext& Context)
 	}
 
 	FChunkOutput Output(HostGroupName, Platform, BuildVersion);
-	Output.Reserve(ContainerFilenames.Num(), Input.Settings.PackageSets.Num());
 
 	for (const FString& Filename : ContainerFilenames)
 	{
@@ -783,15 +546,15 @@ static int32 ChunkPluginCommandEntry(const FContext& Context)
 		// TODO: At the moment the only way we can identify an OnDemand container is by the name. Might be better to pass this info into the tool
 		const bool bIsStreamOnDemandContainer = FullPath.Contains(TEXT("OnDemand"));
 
-		EOnDemandTocFlags ContainerFlags = EOnDemandTocFlags::None;
+		EOnDemandContainerEntryFlags ContainerFlags = EOnDemandContainerEntryFlags::None;
 
 		if (bSupportStreamingContainers && bIsStreamOnDemandContainer)
 		{
-			ContainerFlags = EOnDemandTocFlags::StreamOnDemand;
+			ContainerFlags = EOnDemandContainerEntryFlags::StreamOnDemand;
 		}
 		else if (bSupportInstalledContainers && !bIsStreamOnDemandContainer)
 		{
-			ContainerFlags = EOnDemandTocFlags::InstallOnDemand;
+			ContainerFlags = EOnDemandContainerEntryFlags::InstallOnDemand;
 		}
 		else
 		{
@@ -805,19 +568,6 @@ static int32 ChunkPluginCommandEntry(const FContext& Context)
 		}
 	}
 
-	check(bSupportInstalledContainers || Output.FoundTagSets.IsEmpty());
-
-	if (bSupportInstalledContainers)
-	{
-		Output.InstallOnDemandToc.TagSets.Reserve(Output.FoundTagSets.Num());
-		for (TPair<FString, TArray<FOnDemandTocTagSetPackageList>>& Pair : Output.FoundTagSets)
-		{
-			FOnDemandTocTagSet& TocTagSet = Output.InstallOnDemandToc.TagSets.Emplace_GetRef();
-			TocTagSet.Tag = Pair.Key;
-			TocTagSet.Packages = MoveTemp(Pair.Value);
-		}
-	}
-
 	if (bDeleteContainerFiles)
 	{
 		TRACE_CPUPROFILER_EVENT_SCOPE(ChunkPluginCommand::DeleteContainerFiles);
@@ -843,25 +593,22 @@ static int32 ChunkPluginCommandEntry(const FContext& Context)
 		UE_LOGF(LogIoStoreOnDemand, Warning, "Adding additional file(s) to the on-demand container TOC is no longer supported");
 	}
 
-	const int64 Timestamp = FDateTime::UtcNow().ToUnixTimestamp();
 	// todo, we actually want the .uondemandtoc files in the base directory
 	// const FString TocPath = OutputFolder / Filename; // like this yo
 	const FString OnDemandTocOutputFolder = OutputFolder / IoStoreRelativeFolder;
 
-	Output.SetTocsTimestamp(Timestamp);
-
-	if (bSupportInstalledContainers && !Output.InstallOnDemandToc.Containers.IsEmpty())
+	if (bSupportInstalledContainers && !Output.InstallOnDemandToc.IsEmpty())
 	{
-		if (int32 Result = WriteOnDemandToc(EOnDemandTocFlags::InstallOnDemand, Output.InstallOnDemandToc, OnDemandTocName, OnDemandTocOutputFolder, Output.Stats); Result < 0)
+		if (int32 Result = WriteOnDemandToc(EOnDemandContainerEntryFlags::InstallOnDemand, Output.InstallOnDemandToc, OnDemandTocName, OnDemandTocOutputFolder, Output.Stats); Result < 0)
 		{
 			return Result;
 		}
 	}
 
-	if (bSupportStreamingContainers && !Output.StreamOnDemandToc.Containers.IsEmpty())
+	if (bSupportStreamingContainers && !Output.StreamOnDemandToc.IsEmpty())
 	{
 		// TODO: Improve how we choose the name here.
-		if (int32 Result = WriteOnDemandToc(EOnDemandTocFlags::StreamOnDemand, Output.StreamOnDemandToc, OnDemandTocName + TEXT("ondemand"), OnDemandTocOutputFolder, Output.Stats); Result < 0)
+		if (int32 Result = WriteOnDemandToc(EOnDemandContainerEntryFlags::StreamOnDemand, Output.StreamOnDemandToc, OnDemandTocName + TEXT("ondemand"), OnDemandTocOutputFolder, Output.Stats); Result < 0)
 		{
 			return Result;
 		}
@@ -964,9 +711,9 @@ static FCommand ChunkPluginCommand(
 		TArgument<bool>(TEXT("-KeepContainerFiles"),		TEXT("Should we keep the container files after processing them.")),
 		TArgument<bool>(TEXT("-SkipStreamOnDemand"),		TEXT("Do not chunk any stream on-demand containers")),
 		TArgument<bool>(TEXT("-SkipInstallOnDemand"),		TEXT("Do not chunk any install on-demand containers")),
-		TArgument<bool>(TEXT("-IgnoreContainerHeader"),		TEXT("Ignore the container header?")),
 		TArgument<FStringView>(TEXT("-BucketPrefix"),		TEXT("Path to prefix to bucket objects")),
 		TArgument<int32>(TEXT("-MaxConcurrentUploads"),		TEXT("Number of simultaneous uploads")),
+		TArgument<int32>(TEXT("-MaxPartitionSize"),			TEXT("Max partition size in bytes.")),
 		TArgument<bool>(TEXT("-ValidateChunks"),			TEXT("Runs additional validation on chunks after they are read from the container and before they are written out")),
 		S3Arguments,
 	}
diff --git a/Engine/Source/Developer/IoStore/OnDemandUtilities/Private/ConvertToc.cpp b/Engine/Source/Developer/IoStore/OnDemandUtilities/Private/ConvertToc.cpp
--- a/Engine/Source/Developer/IoStore/OnDemandUtilities/Private/ConvertToc.cpp
+++ b/Engine/Source/Developer/IoStore/OnDemandUtilities/Private/ConvertToc.cpp
@@ -15,6 +15,7 @@
 #include "Misc/Paths.h"
 #include "Misc/SecureHash.h"
 #include "Misc/StringBuilder.h"
+#include "Runtime/Experimental/IoStore/OnDemand/Private/OnDemandIoStore.h"
 #include "Serialization/CompactBinarySerialization.h"
 #include "Serialization/CompactBinaryWriter.h"
 #include "Serialization/LargeMemoryReader.h"
@@ -113,10 +114,10 @@ FContainerMeta FContainerMeta::Load(const FString& Directory)
 }
 
 ////////////////////////////////////////////////////////////////////////////////
-TArray<UE::IoStore::Serialization::V1::FOnDemandTocReader> LoadContainerTocs(const FString& PathOrDir)
+TArray<UE::IoStore::Serialization::V2::FOnDemandTocReader> LoadContainerTocs(const FString& PathOrDir)
 {
 	using namespace UE::IoStore::Serialization;
-	using namespace UE::IoStore::Serialization::V1;
+	using namespace UE::IoStore::Serialization::V2;
 
 	IFileManager& Ifm = IFileManager::Get();
 
@@ -159,8 +160,8 @@ TArray<UE::IoStore::Serialization::V1::FOnDemandTocReader> LoadContainerTocs(con
 static int32 ConvertTocCommandEntry(const FContext& Context)
 {
 	using namespace UE::IoStore::Serialization;
-	using namespace UE::IoStore::Serialization::V1;
-	using FChunkEntry = UE::IoStore::Serialization::V1::FOnDemandChunkEntry;
+	using namespace UE::IoStore::Serialization::V2;
+	using FChunkEntry = UE::IoStore::Serialization::V2::FOnDemandChunkEntry;
 
 	GPrintLogCategory = false;
 	GPrintLogVerbosity = false;
@@ -242,6 +243,16 @@ static int32 ConvertTocCommandEntry(const FContext& Context)
 						if (bExcludeChunkEntries == false)
 						{
 							TUtf8StringBuilder<512> Sb;
+							Writer.BeginArray(UTF8TEXTVIEW("PartitionEntries"));
+							for (const FOnDemandPartitionEntry& Partition : ContainerView.PartitionEntries)
+							{
+								Writer.BeginObject();
+								Writer.AddString(UTF8TEXTVIEW("Hash"), LexToString(Partition.Hash));
+								Writer.AddInteger(UTF8TEXTVIEW("Size"), Partition.Size);
+								Writer.EndObject();
+							}
+							Writer.EndArray();
+
 							Writer.BeginArray(UTF8TEXTVIEW("ChunkEntries"));
 							for (int32 Idx = 0; const FIoChunkId& ChunkId : ContainerView.ChunkIds)
 							{
diff --git a/Engine/Source/Developer/IoStore/OnDemandUtilities/Private/ResaveToc.cpp b/Engine/Source/Developer/IoStore/OnDemandUtilities/Private/ResaveToc.cpp
deleted file mode 100644
diff --git a/Engine/Source/Developer/IoStore/OnDemandUtilities/Private/Upload.cpp b/Engine/Source/Developer/IoStore/OnDemandUtilities/Private/Upload.cpp
--- a/Engine/Source/Developer/IoStore/OnDemandUtilities/Private/Upload.cpp
+++ b/Engine/Source/Developer/IoStore/OnDemandUtilities/Private/Upload.cpp
@@ -1,41 +1,703 @@
 // Copyright Epic Games, Inc. All Rights Reserved.
 
+#include "Upload.h"
+
 #include "Command.h"
 #include "Common.h"
+#include "UploadQueue.h"
 
-#include "Async/Async.h"
-#include "Async/Mutex.h"
-#include "Containers/Queue.h"
 #include "HAL/FileManager.h"
-#include "IO/IoBuffer.h"
+#include "HAL/Platform.h"
 #include "IO/IoChunkEncoding.h"
 #include "IO/IoContainerHeader.h"
-#include "IO/IoStore.h"
+#include "IO/IoDispatcher.h"
 #include "IO/IoStoreOnDemand.h"
 #include "IO/Serialization/OnDemandContainerToc.h"
+#include "Misc/Compression.h"
 #include "Misc/ConfigCacheIni.h"
-#include "Misc/FileHelper.h"
 #include "Misc/KeyChainUtilities.h"
 #include "Misc/Paths.h"
 #include "Misc/PathViews.h"
 #include "S3/S3Client.h"
+#include "Serialization/CompactBinary.h"
 #include "Serialization/CompactBinarySerialization.h"
 #include "Serialization/CompactBinaryWriter.h"
 #include "Serialization/MemoryReader.h"
 #include "Serialization/LargeMemoryWriter.h"
-#include "UploadQueue.h"
 
-namespace UE::IoStore::Tool
+namespace UE::IoStore::Tool { extern FArgumentSet S3Arguments; };
+
+namespace UE::IoStore::Tool::Upload
+{
+
+using EOnDemandContainerEntryFlags = UE::IoStore::Serialization::EOnDemandContainerEntryFlags;
+
+////////////////////////////////////////////////////////////////////////////////
+FIoStatus LoadContainers(
+	TConstArrayView<FString> FilePaths,
+	const FEncryptionKeys& EncryptionKeys,
+	TPagedArray<FContainerData>& OutContainers,
+	bool bIgnoreContainerFlags)
+{
+	for (const FString& FilePath : FilePaths)
+	{
+		TUniquePtr<FIoStoreReader> MaybeReader = MakeUnique<FIoStoreReader>();
+		if (FIoStatus Status = MaybeReader->Initialize(
+			*FPaths::ChangeExtension(FilePath, TEXT("")),
+			EncryptionKeys); !Status.IsOk())
+		{
+			return FIoStatusBuilder(EIoErrorCode::FileOpenFailed)
+				<< TEXT("Failed to open container '") << FilePath << TEXT("'");
+		}
+
+		if (bIgnoreContainerFlags == false && 
+			EnumHasAnyFlags(MaybeReader->GetContainerFlags(), EIoContainerFlags::OnDemand) == false)
+		{
+			continue;
+		}
+
+		FContainerData& Container	= OutContainers.Add_GetRef(FContainerData());
+		Container.FilePath			= FilePath;
+		Container.Name				= FPaths::GetBaseFilename(FilePath);
+		Container.Reader			= MoveTemp(MaybeReader);
+		FIoStoreReader& Reader		= *Container.Reader;
+
+		Reader.EnumerateChunks([&Container](FIoStoreTocChunkInfo&& Info)
+		{ 
+			Container.ChunkInfos.Emplace(MoveTemp(Info));
+			return true;
+		});
+
+		{
+			const bool bDecrypt = false;
+			const FIoChunkId ChunkId = CreateContainerHeaderChunkId(Reader.GetContainerId());
+			TIoStatusOr<FIoStoreCompressedReadResult> ReadResult = Reader.ReadCompressed(ChunkId, FIoReadOptions(), bDecrypt);
+			if (ReadResult.IsOk())
+			{
+				Container.Header = ReadResult.ValueOrDie().IoBuffer;
+			}
+			else
+			{
+				const EIoErrorCode Err = ReadResult.Status().GetErrorCode();
+				if (Err != EIoErrorCode::UnknownChunkID && Err != EIoErrorCode::NotFound)
+				{
+					return FIoStatusBuilder(EIoErrorCode::ReadError)
+						<< TEXT("Failed to read container header '") << FilePath << TEXT("'");
+				}
+			}
+		}
+
+		{
+			const FString UTocFilePath = FPaths::ChangeExtension(FilePath, TEXT(".utoc"));
+			TUniquePtr<FArchive> Ar(IFileManager::Get().CreateFileReader(*UTocFilePath));
+
+			if (!Ar || Ar->IsError())
+			{
+				return FIoStatusBuilder(EIoErrorCode::FileOpenFailed)
+					<< TEXT("Failed to read container TOC '") << UTocFilePath << TEXT("'");
+			}
+
+			Container.UToc = FIoBuffer(Ar->TotalSize());
+			Ar->Serialize(Container.UToc.GetData(), Ar->TotalSize());
+
+			if (Ar->Close() == false)
+			{
+				return FIoStatusBuilder(EIoErrorCode::ReadError)
+					<< TEXT("Failed to read container TOC '") << UTocFilePath << TEXT("'");
+			}
+
+			Container.UTocHash = FIoHash::HashBuffer(Container.UToc.GetView());
+		}
+	}
+
+	return FIoStatus::Ok;
+}
+
+////////////////////////////////////////////////////////////////////////////////
+FIoStatus LoadTagSets(
+	const FString& FilePath,
+	FTagSets& Out)
+{
+	IFileManager& Ifm = IFileManager::Get();
+
+	if (FilePath.IsEmpty() || Ifm.FileExists(*FilePath) == false)
+	{
+		return FIoStatus(EIoErrorCode::NotFound);
+	}
+
+	TUniquePtr<FArchive> Ar(Ifm.CreateFileReader(*FilePath));
+	if (!Ar || Ar->IsError())
+	{
+		return FIoStatus(
+			EIoErrorCode::ReadError,
+			FString::Printf(TEXT("Failed to open tag set file '%s'"), *FilePath));
+	}
+
+	TSharedRef<TJsonReader<UTF8CHAR>> JsonReader = TJsonReaderFactory<UTF8CHAR>::Create(Ar.Get());
+
+	TSharedPtr<FJsonValue> JsonValue;
+	if (!FJsonSerializer::Deserialize(*JsonReader, JsonValue))
+	{
+		return FIoStatus(
+			EIoErrorCode::ReadError,
+			FString::Printf(TEXT("Failed to read tag set file '%s'"), *FilePath));
+	}
+
+	TSharedPtr<FJsonObject> PackageSetsObject = JsonValue->AsObject();
+	if (!PackageSetsObject)
+	{
+		return FIoStatus(
+			EIoErrorCode::ReadError,
+			FString::Printf(TEXT("Invalid JSON object in tag set file '%s'"), *FilePath));
+	}
+
+	for (const TPair<UE::FSharedString, TSharedPtr<FJsonValue>>& Kv : PackageSetsObject->Values)
+	{
+		TArray<FString> PackageNames;
+		if (!PackageSetsObject->TryGetStringArrayField(Kv.Key, PackageNames))
+		{
+			return FIoStatus(
+				EIoErrorCode::ReadError,
+				FString::Printf(TEXT("Invalid JSON value at '%s' in tag set file '%s'"),
+					*Kv.Key, *FilePath));
+		}
+
+		Out.Add(FString(Kv.Key), MoveTemp(PackageNames));
+	}
+
+	return FIoStatus::Ok;
+}
+
+////////////////////////////////////////////////////////////////////////////////
+FIoStatus CreateTagSets(
+	const FTagSets& TagSets,
+	FContainerData& Container)
+{
+	if (Container.Header.GetSize() == 0)
+	{
+		return FIoStatus(EIoErrorCode::InvalidParameter);
+	}
+
+	FIoContainerHeader ContainerHeader;
+	{
+		FMemoryReaderView Ar(Container.Header.GetView());
+		Ar << ContainerHeader;
+		if (Ar.Close() == false)
+		{
+			return FIoStatus(
+				EIoErrorCode::ReadError,
+				FString::Printf(TEXT("Failed to deserialize header chunk for container '%s'"),
+					*Container.FilePath));
+		}
+	}
+
+	for (const TPair<FString, TArray<FString>>& Kv : TagSets)
+	{
+		TConstArrayView<FString>	PackageNames = Kv.Value;
+		TArray<uint32>				Indices;
+
+		for (const FString& PackageName : PackageNames)
+		{
+			const FPackageId	PackageId = FPackageId::FromName(FName(PackageName));
+			int32				Index = INDEX_NONE;
+
+			if (ContainerHeader.PackageIds.Find(PackageId, Index))
+			{
+				check(Index >= 0);
+				Indices.Add(uint32(Index));
+			}
+		}
+
+		if (Indices.IsEmpty() == false)
+		{
+			Container.TagSets.Add(Kv.Key, MoveTemp(Indices));
+		}
+	}
+
+	return FIoStatus::Ok;
+}
+
+////////////////////////////////////////////////////////////////////////////////
+void GenerateTestChunks(
+	FS3Client& Client,
+	const FString& Bucket,
+	const FString& BucketRelativePath,
+	uint8 MaxExp)
+{
+	using namespace UE::IoStore::Serialization;
+
+	// Generate chunks in the form of 1kib.iochunk - 32768kib.iochunk
+	const int32 Exp = FMath::Clamp(int32(MaxExp), 1, 16);
+
+	TStringBuilder<256> Key;
+	for (int32 Shft = 0; Shft < Exp; ++Shft)
+	{
+		Key.Reset();
+		const int32 ChunkId		= (1 << Shft);
+		const int32 ChunkSize	= ChunkId << 10;
+		Key << BucketRelativePath << TEXT("/") << ChunkId << TEXT("kib") << FOnDemandFileExt::Chunk;
+
+		const FS3HeadObjectResponse HeadResponse = Client.HeadObject(FS3HeadObjectRequest{Bucket, Key.ToString()});
+		if (HeadResponse.IsOk())
+		{
+			continue;
+		}
+
+		FIoBuffer Buffer(ChunkSize);
+		ANSICHAR* BufferContent = reinterpret_cast<ANSICHAR*>(Buffer.GetData());
+		FAnsiString SizeText = FAnsiString::Printf("%d", ChunkSize);
+		for (int32 Idx = 0; Idx < Buffer.GetSize(); ++Idx)
+		{
+			BufferContent[Idx] = SizeText[Idx % SizeText.Len()];
+		}
+
+		const FS3PutObjectResponse Response = Client.TryPutObject(FS3PutObjectRequest
+		{
+			.BucketName = Bucket,
+			.Key		= Key.ToString(),
+			.ObjectData = Buffer.GetView()
+		}); 
+
+		if (Response.IsOk())
+		{
+			UE_LOGF(LogIoStoreOnDemand, Display, "Uploaded test chunk '%ls'", Key.ToString());
+		}
+		else
+		{
+			UE_LOGF(LogIoStoreOnDemand, Display, "Failed to upload test chunk '%ls'", Key.ToString());
+		}
+	}
+}
+
+////////////////////////////////////////////////////////////////////////////////
+FIoStatus DownloadPartitionChunkHashes(
+	FS3Client& Client,
+	const FString& Bucket,
+	const FString& BucketRelativePath,
+	int32 MaxTocListCount,
+	int32 MaxTocDownloadCount,
+	TSet<FIoHash>& Out)
 {
+	if (Bucket.IsEmpty())
+	{
+		return FIoStatus(EIoErrorCode::InvalidParameter, TEXT("Invalid bucket name"));
+	}
+
+	TStringBuilder<256> Key;
+	if (!BucketRelativePath.IsEmpty())
+	{
+		Key << BucketRelativePath << TEXT("/");
+	}
+
+	const int32 MaxKeys				= 1000; // AWS max value
+	const int32 MaxPaginationCount	= FMath::Max(1, MaxTocListCount / MaxKeys);
+	TArray<FS3Object> Objects;
+	FString Marker;
+
+	for (int32 Pagination = 0; Pagination < MaxPaginationCount; ++Pagination)
+	{
+		UE_LOGF(LogIoStoreOnDemand, Display, "List #%d/%d '%ls/%ls/%ls' TocCount=%d, MaxTocCount=%d",
+			Pagination + 1, MaxPaginationCount, *Client.GetConfig().ServiceUrl,
+			*Bucket, Key.ToString(), Objects.Num(), MaxTocListCount);
+
+		FS3ListObjectResponse Response = Client.ListObjects(FS3ListObjectsRequest
+		{
+			.BucketName = Bucket,
+			.Prefix		= Key.ToString(),
+			.Delimiter	= TCHAR('/'),
+			.MaxKeys	= MaxKeys,
+			.Marker		= Marker
+		});
+
+		Marker = MoveTemp(Response.NextMarker);
+
+		for (FS3Object& Object : Response.Objects)
+		{
+			if (Object.Key.EndsWith(UE::IoStore::Serialization::FOnDemandFileExt::PartitionToc))
+			{
+				Objects.Add(MoveTemp(Object));
+			}
+		}
 
-#define SHOULD_UPLOAD 1
+		if (Response.IsOk() == false || Response.Objects.IsEmpty() || Response.bIsTruncated == false || Objects.Num() >= MaxTocListCount)
+		{
+			break;
+		}
+	}
+
+	{
+		// Just to be sure
+		const int32 TotalObjectCount = Objects.Num();
+		TSet<FS3Object> UniqueObjects(Objects);
+		Objects = UniqueObjects.Array();
+
+		if (Objects.Num() < TotalObjectCount)
+		{
+			UE_LOGF(LogIoStoreOnDemand, Display, "Removed %d non unique TOC objects", TotalObjectCount - Objects.Num());
+		}
+	}
+
+	Objects.Sort([](const FS3Object& LHS, const FS3Object& RHS) { return LHS.LastModified > RHS.LastModified; });
+
+	uint64 TotalExistingTocs		= 0;
+	const int32 TocDownloadCount	= FMath::Min(Objects.Num(), MaxTocDownloadCount);
+
+	if (TocDownloadCount == 0)
+	{
+		return FIoStatus::Ok;
+	}
+
+	UE_LOGF(LogIoStoreOnDemand, Display, "Fetching %d latest TOC file(s) ...", TocDownloadCount);
+
+	for (int32 Idx = 0; Idx < TocDownloadCount; ++Idx)
+	{
+		const FS3Object& TocInfo = Objects[Idx];
+		if (TocInfo.Key.EndsWith(UE::IoStore::Serialization::FOnDemandFileExt::PartitionToc) == false)
+		{
+			continue;
+		}
+
+		UE_LOGF(LogIoStoreOnDemand, Display, "Fetching TOC %d/%d '%ls/%ls/%ls', Size=%llu, LastModified=%ls",
+			Idx + 1, TocDownloadCount, *Client.GetConfig().ServiceUrl, *Bucket,
+			*TocInfo.Key, TocInfo.Size, *TocInfo.LastModifiedText);
+
+		FS3GetObjectResponse TocResponse = Client.GetObject(FS3GetObjectRequest{ Bucket, TocInfo.Key });
+		if (TocResponse.IsOk() == false)
+		{
+			UE_LOGF(LogIoStoreOnDemand, Warning, "Failed to fetch TOC '%ls/%ls/%ls'",
+				*Client.GetConfig().ServiceUrl, *Bucket, *TocInfo.Key);
+			continue;
+		}
+
+		FCbObjectView	TocView(TocResponse.GetBody().GetData());
+		FCbArrayView	HashesView = TocView[UTF8TEXT("Hashes")].AsArrayView();
+
+		for (FCbFieldView Item : HashesView)
+		{
+			const FIoHash& Hash = Item.AsHash();
+			Out.Add(Hash);
+		}
+
+		TotalExistingTocs++;
+	}
+
+	return FIoStatus::Ok;
+}
+
+////////////////////////////////////////////////////////////////////////////////
+FIoStatus UploadPartitionChunkHashes(
+	FS3Client& Client,
+	const FString& Bucket,
+	const FString& BucketRelativePath,
+	const FString& TargetPlatform,
+	const FString& BuildVersion,
+	TConstArrayView<FPartitionInfo> Partitions)
+{
+	FCbWriter Writer;
+	Writer.BeginObject();
+	{
+		Writer.BeginObject(UTF8TEXT("Header"));
+		{
+			Writer << UTF8TEXT("Timestamp") << FDateTime::Now().ToUnixTimestamp();
+			Writer << UTF8TEXT("TargetPlatform") << TargetPlatform;
+			Writer << UTF8TEXT("BuildVersion") << BuildVersion;
+		}
+		Writer.EndObject();
+
+		Writer.BeginArray(UTF8TEXT("Hashes"));
+		for (const FPartitionInfo& Info : Partitions)
+		{
+			Writer.AddHash(Info.Hash);
+		}
+		Writer.EndArray();
+	}
+	Writer.EndObject();
+
+	FLargeMemoryWriter Ar;
+	Writer.Save(Ar);
+	const FIoHash ObjHash = FIoHash::HashBuffer(Ar.GetView());
+
+	TStringBuilder<256> Key;
+	if (!BucketRelativePath.IsEmpty())
+	{
+		Key << BucketRelativePath << TEXT("/");
+	}
+	Key << ObjHash << UE::IoStore::Serialization::FOnDemandFileExt::PartitionToc;
+
+	const FS3PutObjectResponse Response = Client.TryPutObject(FS3PutObjectRequest{ Bucket, Key.ToString(), Ar.GetView() });
+	if (Response.IsOk() == false)
+	{
+		return FIoStatus(
+			EIoErrorCode::WriteError,
+			FString::Printf(TEXT("Failed to upload partition hash(s), reason: %s"),
+				*Response.GetErrorStatus()));
+	}
+
+	return FIoStatus::Ok;
+}
+
+////////////////////////////////////////////////////////////////////////////////
+static bool ShouldValidateChunk(
+	const FChunkValidationSettings& ValidationSettings,
+	FName CompressionFormat,
+	uint64 RawChunkSize)
+{
+	static TMap<FName, bool> ValidatedFormats;
+
+	if (ValidationSettings.bValidate == false)
+	{
+		return false;
+	}
+
+	if (!CompressionFormat.IsNone())
+	{
+		if (bool* bCachedResult = ValidatedFormats.Find(CompressionFormat))
+		{
+			if (*bCachedResult == false)
+			{
+				return false;
+			}
+		}
+		else
+		{
+			// Note that if IsFormatValid returns false it will log a message about it. This is one of the
+			// reasons we must make sure to only check once per format to avoid  log spam. The other reason
+			// is that the check is fairly slow.
+			bool bValid = FCompression::IsFormatValid(CompressionFormat);
+			ValidatedFormats.Add(CompressionFormat, bValid);
+
+			if (bValid == false)
+			{
+				UE_LOGF(LogIoStoreOnDemand, Display, "Skipping validation tests for chunks with the compression format '%ls'", *CompressionFormat.ToString())
+				return false;
+			}
+		}
+	}
+
+	return RawChunkSize < ValidationSettings.MaxRawChunkSize;
+}
+
+////////////////////////////////////////////////////////////////////////////////
+FIoStatus UploadContainer(
+	FContainerData& Container,
+	const FEncryptionKeys& EncryptionKeys,
+	int32 MaxPartitionSize,
+	UE::IoStore::Serialization::EOnDemandContainerEntryFlags ContainerFlags,
+	const TSet<FIoHash>& ExistingChunks,
+	FUploadChunkCallback&& Upload,
+	UE::IoStore::Serialization::FOnDemandTocWriter& OutWriter,
+	const FChunkValidationSettings& ValidationSettings)
+{
+	using namespace UE::IoStore::Serialization;
+
+	check(Container.Reader.IsValid());
+	FIoStoreReader& Reader = *Container.Reader;
+
+	FIoBuffer	DecodedBuffer;
+	FName		ContainerCompressionFormat	= NAME_None;
+	int32		Partition					= INDEX_NONE;
+	uint64		PartitionSize				= 0;
+
+	auto EndPartition = [&Container, &ExistingChunks, &Upload, &OutWriter, &Partition, &PartitionSize]() -> FIoStatus
+	{
+		FIoHash			Hash;
+		FIoBuffer		Chunk		= OutWriter.EndPartition(Partition, Hash);
+		const FString	HashString	= LexToString(Hash);
+		Partition					= INDEX_NONE;
+		PartitionSize				= 0;
+
+		Container.Stats.Partitions.Count++;
+		Container.Stats.Partitions.Size += Chunk.GetSize();
+
+		if (ExistingChunks.Contains(Hash))
+		{
+			return FIoStatus::Ok;
+		}
+
+		if (FIoStatus Status = Upload(Chunk, Hash); !Status.IsOk())
+		{
+			return Status;
+		}
+
+		Container.UploadedPartitions.Add(FPartitionInfo
+		{
+			.Hash = Hash,
+			.Size = Chunk.GetSize()
+		});
+
+		Container.Stats.UploadedPartitions.Count++;
+		Container.Stats.UploadedPartitions.Size += Chunk.GetSize();
+
+		return FIoStatus::Ok;
+	};
+
+	OutWriter.BeginContainer(
+		Container.Name,	
+		Reader.GetContainerId(),
+		Container.UTocHash,
+		Reader.GetEncryptionKeyGuid(),
+		Reader.GetCompressionBlockSize(),
+		uint32(Reader.GetContainerFlags()),
+		ContainerFlags,
+		Container.Header);
+
+	for (const FIoStoreTocChunkInfo& ChunkInfo : Container.ChunkInfos)
+	{
+		const bool bDecrypt = false;
+		TIoStatusOr<FIoStoreCompressedReadResult> ReadStatus = Reader.ReadCompressed(
+			ChunkInfo.Id,
+			FIoReadOptions(),
+			bDecrypt);
+
+		if (!ReadStatus.IsOk())
+		{
+			return ReadStatus.Status();
+		}
+
+		FIoStoreCompressedReadResult	ReadResult = ReadStatus.ConsumeValueOrDie();
+		FMemoryView						EncodedBlocks = ReadResult.IoBuffer.GetView();
+		TArray<uint32>					BlockSizes;
+		TArray<FIoBlockHash>			BlockHashes;
+		uint64							RawChunkSize = 0;
+		uint64							EncodedChunkSize = 0;
+		FName							CompressionFormat = NAME_None;
+
+		for (const FIoStoreCompressedBlockInfo& BlockInfo : ReadResult.Blocks)
+		{
+			check(Align(BlockInfo.CompressedSize, FAES::AESBlockSize) == BlockInfo.AlignedSize);
+			const uint64 EncodedBlockSize = BlockInfo.AlignedSize;
+			BlockSizes.Add(uint32(BlockInfo.CompressedSize));
+			
+			FMemoryView EncodedBlock = EncodedBlocks.Left(EncodedBlockSize);
+			EncodedBlocks += EncodedBlock.GetSize();
+			BlockHashes.Add(FIoChunkEncoding::HashBlock(EncodedBlock));
+
+			EncodedChunkSize += EncodedBlockSize;
+			RawChunkSize += BlockInfo.UncompressedSize;
+
+			if (CompressionFormat.IsNone() && BlockInfo.CompressionMethod != NAME_None)
+			{
+				CompressionFormat = BlockInfo.CompressionMethod;
+			}
+		}
+
+		// Assume only a single compression format
+		if (ContainerCompressionFormat.IsNone())
+		{
+			ContainerCompressionFormat = CompressionFormat;
+		}
+		else if (!CompressionFormat.IsNone() && CompressionFormat != ContainerCompressionFormat)
+		{
+			return FIoStatus(EIoErrorCode::InvalidParameter, TEXT("Container contains multiple compression format(s)"));
+		}
+
+		if (EncodedChunkSize != ReadResult.IoBuffer.GetSize())
+		{
+			return FIoStatus(EIoErrorCode::ReadError, TEXT("Encoded chunk size does not match buffer"));
+		}
+
+		// At runtime we are limited to MAX_uint32 for chunk lengths to save space and because anything larger than that
+		// is not reasonable to load via IoStoreOnDemand anyway. So we need to check for this now and fail the upload if
+		// there is a chunk that will fail at runtime.
+		// Note that EncodedSize should always be <= RawChunkSize but test both to be safe.
+		if (RawChunkSize > MAX_uint32 || EncodedChunkSize > MAX_uint32)
+		{
+			return FIoStatus(EIoErrorCode::InvalidParameter, WriteToString<512>(*ChunkInfo.FileName, TEXT(": Chunk size should not exceed MAX_uint32")));
+		}
+
+		if (ShouldValidateChunk(ValidationSettings, CompressionFormat, RawChunkSize))
+		{
+			FIoChunkDecodingParams Params;
+			Params.CompressionFormat	= CompressionFormat;
+			Params.TotalRawSize			= RawChunkSize;
+			Params.EncodedBlockSize		= BlockSizes; 
+			Params.BlockHash			= BlockHashes;
+
+			if (EnumHasAnyFlags(Reader.GetContainerFlags(), EIoContainerFlags::Encrypted))
+			{
+				const FAES::FAESKey* EncryptionKey = EncryptionKeys.Find(Reader.GetEncryptionKeyGuid());
+				if (EncryptionKey == nullptr)
+				{
+					return FIoStatus(EIoErrorCode::InvalidParameter, TEXT("Encryption key not found"));
+				}
+
+				Params.EncryptionKey = FMemoryView(EncryptionKey->Key, EncryptionKey->KeySize);
+			}
+
+			if (DecodedBuffer.GetSize() < RawChunkSize)
+			{
+				DecodedBuffer = FIoBuffer(FMath::RoundUpToPowerOfTwo64(RawChunkSize));
+			}
+
+			FMemoryView			Src = ReadResult.IoBuffer.GetView();
+			FMutableMemoryView	Dst = DecodedBuffer.GetMutableView().Left(RawChunkSize);
+
+			if (FIoChunkEncoding::Decode(Params, Src, Dst, EIoDecodeFlags::None) == false)
+			{
+				return FIoStatus(
+					EIoErrorCode::CompressionError,
+					FString::Printf(TEXT("Failed to decode chunk, Container='%s', ChunkId='%s'"),
+						*Container.Name, *LexToString(ChunkInfo.Id)));
+			}
+		}
+
+		// Check if the chunk fits in the current partition. If the partition size is zero
+		// there will be a one to one mapping beteween partitions and chunks and that is supported.
+		if (PartitionSize > 0 && (PartitionSize + EncodedChunkSize >= MaxPartitionSize))
+		{
+			if (FIoStatus Status = EndPartition(); !Status.IsOk())
+			{
+				return Status;
+			}
+		}
+
+		if (Partition == INDEX_NONE)
+		{
+			Partition = OutWriter.BeginPartition();
+			check(Partition >= 0);
+			check(PartitionSize == 0);
+		}
+
+		PartitionSize = OutWriter.AddChunk(
+			Partition,
+			ChunkInfo.Id,
+			ReadResult.IoBuffer,
+			BlockSizes,
+			BlockHashes,
+			RawChunkSize);
+
+		Container.Stats.Chunks.Count++;
+		Container.Stats.Chunks.Size += EncodedChunkSize;
+	}
+
+	// Flush remaining chunks into the current super chunk
+	if (Partition != INDEX_NONE)
+	{
+		if (FIoStatus Status = EndPartition(); !Status.IsOk())
+		{
+			return Status;
+		}
+	}
+
+	check(Partition == INDEX_NONE);
+	check(PartitionSize == 0);
+
+	for (const TPair<FString, TArray<uint32>>& Kv : Container.TagSets)
+	{
+		check(!Kv.Key.IsEmpty());
+		check(!Kv.Value.IsEmpty());
+		OutWriter.AddTagSet(Kv.Key, Kv.Value);
+	}
+
+	OutWriter.EndContainer(ContainerCompressionFormat);
+
+	return FIoStatus::Ok;
+}
 
 ////////////////////////////////////////////////////////////////////////////////
 struct FUploadParams
 {
 	FString ServiceUrl;
-	FString DistributionUrl;
-	FString FallbackUrl;
 	FString Bucket;
 	FString BucketPrefix;
 	FString Region; 
@@ -46,40 +708,26 @@ struct FUploadParams
 	FString CredentialsFileKeyName;
 	FString BuildVersion;
 	FString TargetPlatform;
-	FString EncryptionKeyName;
 	FString HostGroupName;
 	FString GlobalContainerPath;
 	FString TagSetFile;
 	int32 MaxConcurrentUploads = 16;
 	int32 MaxTocListCount = 10000;
 	int32 MaxTocDownloadCount = 100;
-	EOnDemandTocFlags TocFlags = EOnDemandTocFlags::None;
+	int32 MaxPartitionSize = 0;
+	EOnDemandContainerEntryFlags ContainerFlags = EOnDemandContainerEntryFlags::None;
 
 	bool bDeleteContainerFiles = true;
 	bool bDeletePakFiles = true;
 	bool bPerContainerTocs = false;
 	bool bIgnoreContainerFlags = false;
-	bool bIgnoreContainerHeader = false;
-	bool bLegacyToc = false;
-
-	/** If we should write out the .iochunktoc to disk as well as uploading it. */
-	bool bWriteTocToDisk = false;
-	/** Where the .iochunktoc file should be written out. */
-	FString TocOutputDir;
+	bool bValidateChunks = false;
 
 	FIoStatus Validate() const;
 };
 
 FIoStatus FUploadParams::Validate() const
 {
-	if (GlobalContainerPath.IsEmpty())
-	{
-		if (bWriteTocToDisk && TocOutputDir.IsEmpty())
-		{
-			return FIoStatus(EIoErrorCode::InvalidParameter, TEXT("Cmdline param 'WriteToDisk' requires a valid 'ConfigFilePath' param as well"));
-		}
-	}
-
 	if (!AccessKey.IsEmpty() && SecretKey.IsEmpty())
 	{
 		return FIoStatus(EIoErrorCode::InvalidParameter, TEXT("Invalid secret key"));
@@ -102,29 +750,9 @@ FIoStatus FUploadParams::Validate() const
 	return FIoStatus::Ok;
 }
 
-////////////////////////////////////////////////////////////////////////////////
-struct FTocInfo
-{
-	FIoHash Hash = FIoHash::Zero;
-	FString Path;
-	uint64	Size = 0;
-};
-
-////////////////////////////////////////////////////////////////////////////////
-struct FUploadResult
-{
-	FTocInfo	FileToc;
-	FTocInfo	CloudToc;
-	FString		ServiceUrl;
-};
-
 ////////////////////////////////////////////////////////////////////////////////
 static FUploadParams BuildUploadParams(const FContext& Context)
 {
-	static constexpr int32 MinTocListCount		= 1000;
-	static constexpr int32 MaxTocListCount		= 100000;
-	static constexpr int32 MaxTocDownloadCount	= 10000;
-
 	FUploadParams Ret;
 
 	Ret.ServiceUrl				= Context.Get<FStringView>(TEXT("-ServiceUrl"),				Ret.ServiceUrl);
@@ -139,46 +767,36 @@ static FUploadParams BuildUploadParams(const FContext& Context)
 	Ret.BuildVersion			= Context.Get<FStringView>(TEXT("-BuildVersion"),			Ret.BuildVersion);
 	Ret.TargetPlatform			= Context.Get<FStringView>(TEXT("-TargetPlatform"),			Ret.TargetPlatform);
 	Ret.HostGroupName			= Context.Get<FStringView>(TEXT("-HostGroupName"),			Ret.HostGroupName);
-	Ret.TagSetFile				= Context.Get<FStringView>(TEXT("-TagSetFile"),				Ret.TagSetFile);
 	Ret.GlobalContainerPath		= Context.Get<FStringView>(TEXT("-GlobalContainerPath"),	Ret.GlobalContainerPath);
-	Ret.bWriteTocToDisk			= Context.Get<bool>(TEXT("-WriteTocToDisk"),				Ret.bWriteTocToDisk);
+	Ret.TagSetFile				= Context.Get<FStringView>(TEXT("-TagSetFile"),				Ret.TagSetFile);
 	Ret.bPerContainerTocs		= Context.Get<bool>(TEXT("-PerContainerTocs"),				Ret.bPerContainerTocs);
-	Ret.bIgnoreContainerHeader	= Context.Get<bool>(TEXT("-IgnoreContainerHeader"),			Ret.bIgnoreContainerHeader);
-	Ret.bLegacyToc				= Context.Get<bool>(TEXT("-LegacyToc"),						Ret.bLegacyToc);
-	Ret.MaxConcurrentUploads	= Context.Get<int32>(TEXT("-MaxConcurrentUploads"),			Ret.MaxConcurrentUploads);
-	Ret.MaxTocListCount			= Context.Get<int32>(TEXT("-MaxTocListCount"),				Ret.MaxTocListCount);
-	Ret.MaxTocDownloadCount		= Context.Get<int32>(TEXT("-MaxTocDownloadCount"),			Ret.MaxTocDownloadCount);
+	Ret.MaxConcurrentUploads	= Context.Get<int32>(TEXT("-MaxConcurrentUploads"),			FUploadSettings::MaxConcurrentUploads);
+	Ret.MaxTocListCount			= Context.Get<int32>(TEXT("-MaxTocListCount"),				FUploadSettings::MaxTocListCount);
+	Ret.MaxTocDownloadCount		= Context.Get<int32>(TEXT("-MaxTocDownloadCount"),			FUploadSettings::MaxTocDownloadCount);
+	Ret.MaxPartitionSize		= Context.Get<int32>(TEXT("-MaxPartitionSize"),				FUploadSettings::DefaultPartitionSize);
 
 	const bool bStreamOnDemand	= Context.Get<bool>(TEXT("-StreamOnDemand"),				false);
 	const bool bInstallOnDemand	= Context.Get<bool>(TEXT("-InstallOnDemand"),				false);
 
 	if (bStreamOnDemand)
 	{
-		Ret.TocFlags |= EOnDemandTocFlags::StreamOnDemand;
+		Ret.ContainerFlags = EOnDemandContainerEntryFlags::StreamOnDemand;
 	}
 	else if (bInstallOnDemand)
 	{
-		Ret.TocFlags |= EOnDemandTocFlags::InstallOnDemand;
+		Ret.ContainerFlags = EOnDemandContainerEntryFlags::InstallOnDemand;
 	}
 	else
 	{
-		const EOnDemandTocFlags FallbackFlags = EOnDemandTocFlags::StreamOnDemand;
-		UE_LOGF(LogIas, Warning, "No TOC flags was specified, falling back to '%ls'", *LexToString(FallbackFlags));
-		Ret.TocFlags = FallbackFlags;
+		EOnDemandContainerEntryFlags FallbackFlags = EOnDemandContainerEntryFlags::StreamOnDemand;
+		UE_LOGF(LogIoStoreOnDemand, Warning, "No TOC flags was specified, falling back to '%ls'", *LexToString(FallbackFlags));
+		Ret.ContainerFlags = FallbackFlags;
 	}
 
 	Ret.bDeleteContainerFiles	= !Context.Get<bool>(TEXT("-KeepContainerFiles"),			!Ret.bDeleteContainerFiles);
 	Ret.bDeletePakFiles			= !Context.Get<bool>(TEXT("-KeepPakFiles"),					!Ret.bDeletePakFiles);
 	Ret.bIgnoreContainerFlags	= Context.Get<bool>(TEXT("-IgnoreContainerFlags"),			Ret.bIgnoreContainerFlags);
-
-	if (Ret.bWriteTocToDisk)
-	{
-		FStringView Path = Context.Get<FStringView>(TEXT("-ConfigFilePath"));
-		Ret.TocOutputDir = FPathViews::GetPath(Path);
-
-		Ret.DistributionUrl = Context.Get<FStringView>(TEXT("-DistributionUrl"),	Ret.DistributionUrl);
-		Ret.FallbackUrl = Context.Get<FStringView>(TEXT("-FallbackUrl"),			Ret.FallbackUrl);
-	}
+	Ret.bValidateChunks			= Context.Get<bool>(TEXT("-ValidateChunks"),				false);
 
 	if (FIoStatus Validation = Ret.Validate(); !Validation.IsOk())
 	{
@@ -186,8 +804,9 @@ static FUploadParams BuildUploadParams(const FContext& Context)
 		Context.Abort(*Reason);
 	}
 
-	Ret.MaxTocListCount		= FMath::Max(FMath::Min(Ret.MaxTocListCount, MaxTocListCount), MinTocListCount);
-	Ret.MaxTocDownloadCount = FMath::Max(FMath::Min(Ret.MaxTocDownloadCount, MaxTocDownloadCount), 1);
+	Ret.MaxTocListCount		= FMath::Clamp(Ret.MaxTocListCount, FUploadSettings::MinTocListCount, FUploadSettings::MaxTocListCount);
+	Ret.MaxTocDownloadCount = FMath::Clamp(Ret.MaxTocDownloadCount, 1, FUploadSettings::MaxTocDownloadCount);
+	Ret.MaxPartitionSize	= FMath::Clamp(Ret.MaxPartitionSize, 0, FUploadSettings::MaxPartitionSize);
 
 	// Make sure this is a directory path
 	if (FPaths::GetExtension(Ret.GlobalContainerPath).IsEmpty() == false)
@@ -195,106 +814,16 @@ static FUploadParams BuildUploadParams(const FContext& Context)
 		Ret.GlobalContainerPath	= FPaths::GetPath(Ret.GlobalContainerPath);
 	}
 
-	return Ret;
-}
-
-////////////////////////////////////////////////////////////////////////////////
-static bool SaveArrayToFile(TArrayView64<uint8> Data, const TCHAR* Filename, uint32 WriteFlags = FILEWRITE_None)
-{
-	TUniquePtr<FArchive> Ar = TUniquePtr<FArchive>(IFileManager::Get().CreateFileWriter(Filename, WriteFlags));
-	if (!Ar)
-	{
-		return false;
-	}
-
-	Ar->Serialize(static_cast<void*>(Data.GetData()), Data.Num());
-	Ar->Close();
-
-	return !Ar->IsError() && !Ar->IsCriticalError();
-}
-
-////////////////////////////////////////////////////////////////////////////////
-static FIoStatus WriteContainerFiles(FOnDemandToc& OnDemandToc, const TMap<FIoHash, FString>& UTocPaths, bool bLegacyToc)
-{
-	UE_LOGF(LogIas, Display, "Attempting to write out %d '%ls' files", OnDemandToc.Containers.Num(), *FOnDemandToc::FileExt);
-
-	for (int32 ContainerIndex = 0; ContainerIndex < OnDemandToc.Containers.Num(); ++ContainerIndex)
+	bool bPerContainerTocsConfigValue = false;
+	if (GConfig->GetBool(TEXT("Ias"), TEXT("CreatePerContainerTocs"), bPerContainerTocsConfigValue, GEngineIni))
 	{
-		FOnDemandTocContainerEntry& Container = OnDemandToc.Containers[ContainerIndex];
-		const FString* UTocPath = UTocPaths.Find(Container.UTocHash);
-		if (UTocPath == nullptr)
-		{
-			return FIoStatus(EIoErrorCode::Unknown, FString::Printf(TEXT("Could not find the original path for ondemand container '%s'"), *Container.ContainerName));
-		}
-
-		// Create a new FOnDemandToc to be written out and give it ownership of the FOnDemandTocContainerEntry that we want
-		// associated with it.This way we don't need to write any specialized serialization code to split up the input FOnDemandToc
-		// into one per container.
-		FOnDemandToc ContainerToc;
-		ContainerToc.Header = OnDemandToc.Header;
-		ContainerToc.Meta = OnDemandToc.Meta;
-		ContainerToc.Containers.Emplace(MoveTemp(Container));
-		// find all the tagsets for this container
-		for (const FOnDemandTocTagSet& TagSet : OnDemandToc.TagSets)
-		{
-			FOnDemandTocTagSet NewTagSet;
-			for (const FOnDemandTocTagSetPackageList& PackageList : TagSet.Packages)
-			{
-				if (PackageList.ContainerIndex == ContainerIndex)
-				{
-					FOnDemandTocTagSetPackageList& NewTagSetPackageList = NewTagSet.Packages[NewTagSet.Packages.AddDefaulted()];
-
-					NewTagSetPackageList.ContainerIndex = 0;
-					NewTagSetPackageList.PackageIndicies.Append(PackageList.PackageIndicies);
-				}
-			}
-
-			if (NewTagSet.Packages.Num() > 0)
-			{
-				NewTagSet.Tag = TagSet.Tag;
-				ContainerToc.TagSets.Add(NewTagSet);
-			}
-			
-		}
-		
-
-		const FString OutputPath = FPathViews::ChangeExtension(*UTocPath, FOnDemandToc::FileExt);
-
-		TUniquePtr<FArchive> Ar(IFileManager::Get().CreateFileWriter(*OutputPath));
-
-		if (!Ar.IsValid())
-		{
-			return FIoStatus(EIoErrorCode::FileOpenFailed, FString::Printf(TEXT("Failed to open '%s' for write"), *OutputPath));
-		}
-
-		if (bLegacyToc)
+		if (bPerContainerTocsConfigValue)
 		{
-			(*Ar) << ContainerToc;
-			FOnDemandTocSentinel Sentinel;
-			(*Ar) << Sentinel;
-		}
-		else
-		{
-			UE::IoStore::Serialization::FOnDemandTocWriter Writer;
-			if (TIoStatusOr<uint64> Status = Writer.Write(ContainerToc, *Ar); !Status.IsOk())
-			{
-				return Status.Status();
-			}
+			Ret.bPerContainerTocs = true;
 		}
-
-		if (!Ar->Close())
-		{
-			return FIoStatus(EIoErrorCode::WriteError, FString::Printf(TEXT("Failed to write to '%s'"), *OutputPath));
-		}
-
-		UE_LOGF(LogIas, Display, "Wrote ondemand container file '%ls' to disk", *OutputPath);
-
-		// Move the container entry back to OnDemandToc in case we want to use the data structure
-		// beyond this point in the future.
-		Container = MoveTemp(ContainerToc.Containers[0]);
 	}
 
-	return FIoStatus::Ok;
+	return Ret;
 }
 
 ////////////////////////////////////////////////////////////////////////////////
@@ -353,98 +882,31 @@ static void GetTocPath(
 		{
 			Sb << HostSuffix << TEXT("/");
 		}
-		Sb << Bucket << TEXT("/") << BucketPrefix;
-		OutTocPath = Sb;
-	}
-}
-
-////////////////////////////////////////////////////////////////////////////////
-static void WriteConfigFile(
-	const FUploadParams& UploadParams,
-	const FUploadResult& UploadResult,
-	const FKeyChain& KeyChain)
-{
-	// This string can get quite long so don't bother with any inline storage
-	TStringBuilder<0> Sb;
-	Sb << TEXT("[Endpoint]") << TEXT("\r\n");
-
-	if (!UploadParams.DistributionUrl.IsEmpty())
-	{
-		Sb << TEXT("DistributionUrl=\"") << UploadParams.DistributionUrl << TEXT("\"\r\n");
-
-		if (!UploadParams.FallbackUrl.IsEmpty())
-		{
-			Sb << TEXT("FallbackUrl=\"") << UploadParams.FallbackUrl << TEXT("\"\r\n");
-		}
-	}
-	else
-	{
-		Sb << TEXT("ServiceUrl=\"") << UploadResult.ServiceUrl << TEXT("\"\r\n");
-	}
-
-	Sb << TEXT("TocPath=\"") << UploadResult.FileToc.Path << TEXT("\"\r\n");
-
-	// Temporary solution to get replays working with encrypted on demand content
-	{
-		if (!UploadParams.EncryptionKeyName.IsEmpty())
-		{
-			const TCHAR* EncryptionKeyName = *(UploadParams.EncryptionKeyName);
-
-			TOptional<FNamedAESKey> EncryptionKey;
-			for (const TPair<FGuid, FNamedAESKey>& KeyPair: KeyChain.GetEncryptionKeys())
-			{
-				if (KeyPair.Value.Name.Compare(EncryptionKeyName, ESearchCase::IgnoreCase) == 0)
-				{
-					EncryptionKey.Emplace(KeyPair.Value);
-				}
-			}
-
-			if (EncryptionKey)
-			{
-				FString KeyString = FBase64::Encode(EncryptionKey.GetValue().Key.Key, FAES::FAESKey::KeySize);
-				Sb << TEXT("ContentKey=\"") << EncryptionKey.GetValue().Guid.ToString() << TEXT(":") << KeyString << TEXT("\"\r\n");
-			}
-			else
-			{
-				UE_LOGF(LogIoStore, Warning, "Failed to encryption key '%ls' in key chain", EncryptionKeyName);
-			}
-		}
-	}
-
-	FString ConfigFilePath = UploadParams.TocOutputDir / TEXT("IoStoreOnDemand.ini");
-	UE_LOGF(LogIoStore, Display, "Saving on demand config file '%ls'", *ConfigFilePath);
-	if (FFileHelper::SaveStringToFile(Sb, *ConfigFilePath) == false)
-	{
-		UE_LOGF(LogIoStore, Error, "Failed to save on demand config file '%ls'", *ConfigFilePath);
+		Sb << Bucket << TEXT("/") << BucketPrefix;
+		OutTocPath = Sb;
 	}
 }
 
-
-
 ////////////////////////////////////////////////////////////////////////////////
-static TIoStatusOr<FUploadResult> UploadContainerFiles(
+static FIoStatus UploadContainerFiles(
 	const FUploadParams& UploadParams,
 	TConstArrayView<FString> ContainerFiles,
 	const FKeyChain& KeyChain)
 {
+	using namespace UE::IoStore::Serialization;
+	
+	const bool bGlobalToc = UploadParams.GlobalContainerPath.IsEmpty() == false;
+	if (!bGlobalToc && ContainerFiles.Num() > 1)
+	{
+		return FIoStatus(EIoErrorCode::InvalidParameter, TEXT("Multiple container(s) in a single TOC is not allowed."));
+	}
+
 	TMap<FGuid, FAES::FAESKey> EncryptionKeys;
 	for (const TPair<FGuid, FNamedAESKey>& KeyPair: KeyChain.GetEncryptionKeys())
 	{
 		EncryptionKeys.Add(KeyPair.Key, KeyPair.Value.Key);
 	}
 
-	struct FContainerStats
-	{
-		FString ContainerName;
-		uint64 ChunkCount = 0;
-		uint64 TotalBytes = 0;
-		uint64 UploadedChunkCount = 0;
-		uint64 UploadedBytes = 0;
-	};
-	TMap<FString, FContainerStats> ContainerSummary;
-
-	const double StartTime = FPlatformTime::Seconds();
-
 	FS3ClientConfig Config;
 	Config.ServiceUrl = UploadParams.ServiceUrl;
 	Config.Region = UploadParams.Region;
@@ -452,7 +914,7 @@ static TIoStatusOr<FUploadResult> UploadContainerFiles(
 	FS3ClientCredentials Credentials;
 	if (UploadParams.CredentialsFile.IsEmpty() == false)
 	{
-		UE_LOGF(LogIas, Display, "Loading credentials file '%ls'", *UploadParams.CredentialsFile);
+		UE_LOGF(LogIoStoreOnDemand, Display, "Loading credentials file '%ls'", *UploadParams.CredentialsFile);
 		FS3CredentialsProfileStore CredentialsStore = FS3CredentialsProfileStore::FromFile(UploadParams.CredentialsFile);
 		if (CredentialsStore.TryGetCredentials(UploadParams.CredentialsFileKeyName, Credentials) == false)
 		{
@@ -460,7 +922,7 @@ static TIoStatusOr<FUploadResult> UploadContainerFiles(
 		}
 		else
 		{
-			UE_LOGF(LogIas, Display, "Found credentials for '%ls'", *UploadParams.CredentialsFileKeyName);
+			UE_LOGF(LogIoStoreOnDemand, Display, "Found credentials for '%ls'", *UploadParams.CredentialsFileKeyName);
 		}
 	}
 	else
@@ -469,673 +931,244 @@ static TIoStatusOr<FUploadResult> UploadContainerFiles(
 	}
 
 	FS3Client Client(Config, Credentials);
-#if SHOULD_UPLOAD
 	FUploadQueue UploadQueue(Client, UploadParams.Bucket, UploadParams.MaxConcurrentUploads);
-#endif
 
-	if (ContainerFiles.Num() == 0)
+	FString ChunksRelativePath = UploadParams.BucketPrefix.IsEmpty()
+		? TEXT("Chunks")
+		: FString::Printf(TEXT("%s/Chunks"), *UploadParams.BucketPrefix);
+	ChunksRelativePath.ToLowerInline();
+
+	TPagedArray<FContainerData> Containers;
+	if (FIoStatus Status = LoadContainers(
+		ContainerFiles,
+		EncryptionKeys,
+		Containers,
+		UploadParams.bIgnoreContainerFlags); !Status.IsOk())
 	{
-		return FIoStatus(EIoErrorCode::InvalidParameter, TEXT("No container file(s) specified"));
+		return Status;
 	}
 
-	TSet<FIoHash> ExistingChunks;
-	uint64 TotalExistingTocs = 0;
-	uint64 TotalExistingBytes = 0;
+	if (Containers.IsEmpty())
 	{
-		TStringBuilder<256> TocsKey;
-		if (!UploadParams.BucketPrefix.IsEmpty())
-		{
-			TocsKey << UploadParams.BucketPrefix << "/";
-		}
-
-		const int32 MaxKeys = 1000; // AWS max value
-		const int32 MaxPaginationCount = FMath::Max(1, UploadParams.MaxTocListCount / MaxKeys);
-		const int32 MaxTocDownloadCount = UploadParams.MaxTocDownloadCount;
-		TArray<FS3Object> Objects;
-		FString Marker;
-
-		for (int32 Pagination = 0; Pagination < MaxPaginationCount; ++Pagination)
-		{
-			UE_LOGF(LogIas, Display, "List #%d/%d '%ls/%ls/%ls' TocCount=%d, MaxTocCount=%d",
-				Pagination + 1, MaxPaginationCount, *Client.GetConfig().ServiceUrl, *UploadParams.Bucket, TocsKey.ToString(), Objects.Num(), UploadParams.MaxTocListCount);
-			FS3ListObjectResponse Response = Client.ListObjects(FS3ListObjectsRequest
-			{
-				.BucketName = UploadParams.Bucket,
-				.Prefix = TocsKey.ToString(),
-				.Delimiter = TCHAR('/'),
-				.MaxKeys = MaxKeys,
-				.Marker = Marker
-			});
-
-			Marker = MoveTemp(Response.NextMarker);
-
-			for (FS3Object& Object : Response.Objects)
-			{
-				if (Object.Key.EndsWith(TEXT("iochunktoc")))
-				{
-					Objects.Add(MoveTemp(Object));
-				}
-			}
-
-			if (Response.IsOk() == false || Response.Objects.IsEmpty() || Response.bIsTruncated == false || Objects.Num() >= UploadParams.MaxTocListCount)
-			{
-				break;
-			}
-		}
-
-		{
-			// Just to be sure
-			const int32 TotalObjectCount = Objects.Num();
-			TSet<FS3Object> UniqueObjects(Objects);
-			Objects = UniqueObjects.Array();
-			if (Objects.Num() < TotalObjectCount)
-			{
-				UE_LOGF(LogIas, Display, "Removed %d non unique TOC objects", TotalObjectCount - Objects.Num());
-			}
-		}
-
-		Objects.Sort([](const FS3Object& LHS, const FS3Object& RHS) { return LHS.LastModified > RHS.LastModified; });
+		return FIoStatus(EIoErrorCode::InvalidParameter, TEXT("No container file(s) marked as on demand"));
+	}
 
-		const int32 TocDownloadCount = FMath::Min(Objects.Num(), MaxTocDownloadCount);
-		UE_LOGF(LogIas, Display, "Fetching %d latest TOC file(s) ...", TocDownloadCount);
-		for (int32 Idx = 0; Idx < TocDownloadCount; ++Idx)
+	FTagSets TagSets;
+	if (FIoStatus Status = LoadTagSets(UploadParams.TagSetFile, TagSets); Status.IsOk())
+	{
+		UE_LOGF(LogIoStoreOnDemand, Display, "Creating tag set(s) from '%ls'", *UploadParams.TagSetFile);
+		for (FContainerData& Container : Containers)
 		{
-			const FS3Object& TocInfo = Objects[Idx];
-			if (TocInfo.Key.EndsWith(TEXT("iochunktoc")) == false)
-			{
-				continue;
-			}
-
-			UE_LOGF(LogIas, Display, "Fetching TOC %d/%d '%ls/%ls/%ls', Size=%llu, LastModified=%ls", Idx + 1, TocDownloadCount, *Client.GetConfig().ServiceUrl, *UploadParams.Bucket, *TocInfo.Key, TocInfo.Size, *TocInfo.LastModifiedText);
-			FS3GetObjectResponse TocResponse = Client.GetObject(FS3GetObjectRequest
-			{
-				UploadParams.Bucket,
-				TocInfo.Key
-			});
-			
-			if (TocResponse.IsOk() == false)
-			{
-				UE_LOGF(LogIas, Warning, "Failed to fetch TOC '%ls/%ls/%ls'", *Client.GetConfig().ServiceUrl, *UploadParams.Bucket, *TocInfo.Key);
-				continue;
-			}
-
-			FOnDemandToc Toc;
-			FMemoryReaderView Ar(TocResponse.GetBody().GetView());
-			Ar << Toc;
-
-			if (Ar.IsError()) 
-			{
-				Toc = FOnDemandToc{};
-				if (LoadFromCompactBinary(FCbFieldView(TocResponse.GetBody().GetData()), Toc) == false)
-				{
-					UE_LOGF(LogIas, Display, "Failed to load TOC '%ls/%ls/%ls'", *Client.GetConfig().ServiceUrl, *UploadParams.Bucket, *TocInfo.Key);
-					continue;
-				}
-			}
-
-			for (const FOnDemandTocContainerEntry& ContainerEntry : Toc.Containers)
+			// Can only create tag sets for containers with package information
+			if (Container.Header.GetSize() > 0)
 			{
-				for (const FOnDemandTocEntry& TocEntry : ContainerEntry.Entries)
+				if (Status = CreateTagSets(TagSets, Container); !Status.IsOk())
 				{
-					bool bIsAlreadyInSet;
-					ExistingChunks.Add(TocEntry.Hash, &bIsAlreadyInSet);
-					if (!bIsAlreadyInSet)
-					{
-						TotalExistingBytes += TocEntry.EncodedSize;
-					}
+					return Status;
 				}
 			}
-
-			TotalExistingTocs++;
 		}
 	}
-
-	FString ChunksRelativePath = UploadParams.BucketPrefix.IsEmpty() ? TEXT("Chunks") : FString::Printf(TEXT("%s/Chunks"), *UploadParams.BucketPrefix);
-	ChunksRelativePath.ToLowerInline();
-
-	bool bPerContainerTocsConfigValue = false;
-	GConfig->GetBool(TEXT("Ias"), TEXT("CreatePerContainerTocs"), bPerContainerTocsConfigValue, GEngineIni);
-	const bool bWritePerContainerToc = bPerContainerTocsConfigValue || UploadParams.bPerContainerTocs;
-
-	uint64 TotalUploadedChunks = 0;
-	uint64 TotalUploadedBytes = 0;
-
-	TMap<FString, TArray<FString>> PackageSets;
-	if (!UploadParams.TagSetFile.IsEmpty())
+	else if (Status.GetErrorCode() != EIoErrorCode::NotFound)
 	{
-		TUniquePtr<FArchive> TagsetAr(IFileManager::Get().CreateFileReader(*UploadParams.TagSetFile));
-		if (!TagsetAr || TagsetAr->IsError())
-		{
-			UE_LOGF(LogIoStoreOnDemand, Error, "Failed to open settings file '%ls'", *UploadParams.TagSetFile);
-			return FIoStatus(EIoErrorCode::ReadError, FString::Printf(TEXT("Failed to open settings file '%s'"), *UploadParams.TagSetFile));
-		}
-
-		TSharedRef<TJsonReader<UTF8CHAR>> JsonReader = TJsonReaderFactory<UTF8CHAR>::Create(TagsetAr.Get());
-
-		TSharedPtr<FJsonValue> JsonSettings;
-		if (!FJsonSerializer::Deserialize(*JsonReader, JsonSettings))
-		{
-			UE_LOGF(LogIoStoreOnDemand, Error, "Failed to read settings file '%ls'", *UploadParams.TagSetFile);
-			return FIoStatus(EIoErrorCode::ReadError, FString::Printf(TEXT("Failed to read settings file '%s'"), *UploadParams.TagSetFile));
-		}
-
-		TSharedPtr<FJsonObject> PackageSetsObject = JsonSettings->AsObject();
-		if (!PackageSetsObject)
-		{
-			UE_LOGF(LogIoStoreOnDemand, Error, "Bad settings file '%ls'", *UploadParams.TagSetFile);
-			return FIoStatus(EIoErrorCode::ReadError, FString::Printf(TEXT("Bad settings file '%s'"), *UploadParams.TagSetFile));
-		}
-
-
-
-		for (const TPair<UE::FSharedString, TSharedPtr<FJsonValue>>& Pair : PackageSetsObject->Values)
-		{
-			TArray<FString> Packages;
-			if (!PackageSetsObject->TryGetStringArrayField(*Pair.Key, Packages))
-			{
-				UE_LOGF(LogIas, Error, "Bad settings file '%ls'", *UploadParams.TagSetFile);
-				return FIoStatus(EIoErrorCode::ReadError, FString::Printf(TEXT("Bad settings file '%s'"), *UploadParams.TagSetFile));
-			}
-
-			UE_LOGF(LogIas, Display, "Found Package Set '%ls' count %d", *Pair.Key, Packages.Num());
-			/*for (const FString& Package : Packages)
-			{
-				UE_LOGF(LogIas, Display, "-> \t'%ls'", *Package);
-			}*/
-
-			PackageSets.Add(FString(Pair.Key), MoveTemp(Packages));
-		}
+		return Status;
 	}
 
-	FOnDemandToc OnDemandToc;
-	OnDemandToc.Header.Flags = static_cast<uint32>(UploadParams.TocFlags);
-	OnDemandToc.Header.HostGroupName = UploadParams.HostGroupName;
-	OnDemandToc.Containers.Reserve(ContainerFiles.Num());
-
-	// Map of the .utoc paths that we have created ondemand containers for, indexed by their hash so that the paths can
-	// be looked up later. We do not rely on the filename as we cannot be sure that there won't be duplicate file names
-	// stored in different directories.
-	TMap<FIoHash, FString> UTocPaths;
-
-	TArray<FString> FilesToDelete;
-	for (const FString& Path : ContainerFiles)
+	TSet<FIoHash> ExistingChunks;
+	if (FIoStatus Status = DownloadPartitionChunkHashes(
+		Client,
+		UploadParams.Bucket,
+		UploadParams.BucketPrefix,
+		UploadParams.MaxTocListCount,
+		UploadParams.MaxTocDownloadCount,
+		ExistingChunks); !Status.IsOk())
 	{
-		FIoStoreReader ContainerFileReader;
-		{
-			FIoStatus Status = ContainerFileReader.Initialize(*FPaths::ChangeExtension(Path, TEXT("")), EncryptionKeys);
-			if (!Status.IsOk())
-			{
-				UE_LOGF(LogIas, Error, "Failed to open container '%ls' for reading due to: '%ls'", *Path, *Status.ToString());
-				continue;
-			}
-		}
-
-		if (UploadParams.bIgnoreContainerFlags == false &&
-			EnumHasAnyFlags(ContainerFileReader.GetContainerFlags(), EIoContainerFlags::OnDemand) == false)
-		{
-			continue;
-		}
-		UE_LOGF(LogIas, Display, "Uploading ondemand container '%ls'", *Path);
-
-		if (OnDemandToc.Header.BlockSize == 0)
-		{
-			OnDemandToc.Header.BlockSize = ContainerFileReader.GetCompressionBlockSize();
-		}
-		check(OnDemandToc.Header.BlockSize == ContainerFileReader.GetCompressionBlockSize());
-
-		TArray<FIoStoreTocChunkInfo> ChunkInfos;
-		ContainerFileReader.EnumerateChunks([&ChunkInfos](FIoStoreTocChunkInfo&& Info)
-		{ 
-			ChunkInfos.Emplace(MoveTemp(Info));
-			return true;
-		});
-		
-		const int32 ContainerIndex = OnDemandToc.Containers.AddDefaulted();
-		FOnDemandTocContainerEntry& ContainerEntry = OnDemandToc.Containers[ContainerIndex];
-		ContainerEntry.ContainerName = FPaths::GetBaseFilename(Path);
-		ContainerEntry.ContainerId = ContainerFileReader.GetContainerId();
-		ContainerEntry.ContainerFlags = static_cast<uint8>(ContainerFileReader.GetContainerFlags());
-
-
-
-
-
-
-		if (PackageSets.Num() > 0)
-		{
-			TArray<FOnDemandTocTagSet> TagSets;
-
-			const FIoChunkId HeaderChunkId = CreateContainerHeaderChunkId(ContainerFileReader.GetContainerId());
-			TIoStatusOr<FIoBuffer> Status = ContainerFileReader.Read(HeaderChunkId, FIoReadOptions());
-			if (!Status.IsOk())
-			{
-				if (Status.Status().GetErrorCode() == EIoErrorCode::NotFound)
-				{
-					UE_LOGF(LogIas, Display, "No header chunk for container '%ls'", *Path);
-				}
-				else
-				{
-					UE_LOGF(LogIas, Error, "Failed to read header chunk for container '%ls'", *Path);
-					return FIoStatus(EIoErrorCode::ReadError, FString::Printf(TEXT("Failed to read header chunk for container '%s'"), *Path));
-				}
-			}
-			else
-			{
-				FIoContainerHeader ContainerHeader;
-				FMemoryReaderView Ar(Status.ValueOrDie().GetView());
-				Ar << ContainerHeader;
-				Ar.Close();
-				if (Ar.IsError() || Ar.IsCriticalError())
-				{
-					UE_LOGF(LogIas, Error, "Failed to deserialize header chunk for container '%ls'", *Path);
-					return FIoStatus(EIoErrorCode::ReadError, FString::Printf(TEXT("Failed to deserialize header chunk for container '%s'"), *Path));
-				}
-
-				int32 FoundPackageTagsetsCounter = 0;
-				for (const TPair<FString, TArray<FString>>& Pair : PackageSets)
-				{
-					const FString& Tag = Pair.Key;
-					if (Tag == TEXT("AICore"))
-					{
-						static int o = 0;
-						o++;
-					}
-
-					FOnDemandTocTagSet TagSet;
-					TagSet.Tag = Pair.Key;
-
-					FOnDemandTocTagSetPackageList& TagSetPackageList = TagSet.Packages[TagSet.Packages.AddDefaulted()];
-					TagSetPackageList.ContainerIndex = ContainerIndex;
-
-					bool bFoundPackageIndicies = false;
-					for (const FString& Package : Pair.Value)
-					{
-						const FPackageId PackageId = FPackageId::FromName(FName(Package));
-						int32 PackageIndex = INDEX_NONE;
-						if (ContainerHeader.PackageIds.Find(PackageId, PackageIndex))
-						{
-							++FoundPackageTagsetsCounter;
-
-							TagSetPackageList.PackageIndicies.Add(PackageIndex);
-							bFoundPackageIndicies = true;
-						}
-					}
-					UE_LOGF(LogIas, Display, "Found packages %d for tag %ls in container %ls", TagSetPackageList.PackageIndicies.Num(), *Tag, *Path);
-
-					if (bFoundPackageIndicies)
-					{
-						OnDemandToc.TagSets.Add(TagSet);
-					}
-				}
-
-				UE_LOGF(LogIas, Display, "Found tags for %d packages in container %ls", FoundPackageTagsetsCounter, *Path);
-			}
-		}
-
-
+		return Status;
+	}
 
-		if (EnumHasAnyFlags(ContainerFileReader.GetContainerFlags(), EIoContainerFlags::Encrypted))
-		{
-			ContainerEntry.EncryptionKeyGuid = LexToString(ContainerFileReader.GetEncryptionKeyGuid());
-		}
-		FContainerStats& ContainerStats = ContainerSummary.FindOrAdd(ContainerEntry.ContainerName);
-		
-		ContainerEntry.Entries.Reserve(ChunkInfos.Num());
+	TArray<FString>		FilesToDelete;
+	FOnDemandTocWriter	TocWriter;
 
-		for (const FIoStoreTocChunkInfo& ChunkInfo : ChunkInfos)
+	for (FContainerData& Container : Containers)
+	{
+		auto Upload = [&UploadQueue, &ChunksRelativePath](FIoBuffer Chunk, const FIoHash& Hash) -> FIoStatus
 		{
-			const bool bDecrypt = false;
-			TIoStatusOr<FIoStoreCompressedReadResult> Status = ContainerFileReader.ReadCompressed(ChunkInfo.Id, FIoReadOptions(), bDecrypt);
-			if (!Status.IsOk())
-			{
-				return Status.Status();
-			}
-
-			FIoStoreCompressedReadResult ReadResult = Status.ConsumeValueOrDie();
-
-			const uint32 BlockOffset = ContainerEntry.BlockSizes.Num();
-			const uint32 BlockCount = ReadResult.Blocks.Num();
-			const FIoHash ChunkHash = FIoHash::HashBuffer(ReadResult.IoBuffer.GetView());
-
-			FMemoryView EncodedBlocks = ReadResult.IoBuffer.GetView();
-			uint64 RawChunkSize = 0;
-			uint64 EncodedChunkSize = 0;
-			for (const FIoStoreCompressedBlockInfo& BlockInfo : ReadResult.Blocks)
-			{
-				check(Align(BlockInfo.CompressedSize, FAES::AESBlockSize) == BlockInfo.AlignedSize);
-				const uint64 EncodedBlockSize = BlockInfo.AlignedSize;
-				ContainerEntry.BlockSizes.Add(uint32(BlockInfo.CompressedSize));
-
-				FMemoryView EncodedBlock = EncodedBlocks.Left(EncodedBlockSize);
-				EncodedBlocks += EncodedBlock.GetSize();
-				ContainerEntry.BlockHashes.Add(FIoChunkEncoding::HashBlock(EncodedBlock));
-
-				EncodedChunkSize += EncodedBlockSize;
-				RawChunkSize += BlockInfo.UncompressedSize;
-
-				if (OnDemandToc.Header.CompressionFormat.IsEmpty() && BlockInfo.CompressionMethod != NAME_None)
-				{
-					OnDemandToc.Header.CompressionFormat = BlockInfo.CompressionMethod.ToString();
-				}
-			}
+			const FString HashString = LexToString(Hash);
 
-			if (EncodedChunkSize != ReadResult.IoBuffer.GetSize())
-			{
-				return FIoStatus(EIoErrorCode::ReadError, TEXT("Encoded chunk size does not match buffer"));
-			}
+			TStringBuilder<256> Key;
+			Key << ChunksRelativePath
+				<< TEXT("/") << HashString.Left(2)
+				<< TEXT("/") << HashString
+				<< UE::IoStore::Serialization::FOnDemandFileExt::Partition;
 
-			// At runtime we are limited to MAX_uint32 for chunk lengths to save space and because anything larger than that
-			// is not reasonable to load via IoStoreOnDemand anyway. So we need to check for this now and fail the upload if
-			// there is a chunk that will fail at runtime.
-			// Note that EncodedSize should always be <= RawChunkSize but test both to be safe.
-			if (RawChunkSize > MAX_uint32 || EncodedChunkSize > MAX_uint32)
+			if (UploadQueue.Enqueue(Key.ToView(), Chunk) == false)
 			{
-				return FIoStatus(EIoErrorCode::InvalidParameter, WriteToString<512>(*ChunkInfo.FileName, TEXT(": Chunk size should not exceed MAX_uint32")));
+				return FIoStatus(EIoErrorCode::WriteError);
 			}
 
-			FOnDemandTocEntry& TocEntry = ContainerEntry.Entries.AddDefaulted_GetRef();
-			TocEntry.ChunkId = ChunkInfo.Id;
-			TocEntry.Hash = ChunkHash;
-			TocEntry.RawSize = RawChunkSize;
-			TocEntry.EncodedSize = EncodedChunkSize;
-			TocEntry.BlockOffset = BlockOffset;
-			TocEntry.BlockCount = BlockCount;
-
-			ContainerStats.ChunkCount++;
-			ContainerStats.TotalBytes += EncodedChunkSize;
-
-			if (ExistingChunks.Contains(TocEntry.Hash) == false)
-			{
-				const FString HashString = LexToString(ChunkHash);
-
-				TStringBuilder<256> Key;
-				Key << ChunksRelativePath
-					<< TEXT("/") << HashString.Left(2)
-					<< TEXT("/") << HashString
-					<< TEXT(".iochunk");
-
-#if SHOULD_UPLOAD
-				if (UploadQueue.Enqueue(Key, ReadResult.IoBuffer) == false)
-				{
-					return FIoStatus(EIoErrorCode::WriteError, TEXT("Failed to upload chunk"));
-				}
-#endif
-				TotalUploadedChunks++;
-				TotalUploadedBytes += EncodedChunkSize;
-				ContainerStats.UploadedChunkCount++;
-				ContainerStats.UploadedBytes += EncodedChunkSize;
-			}
-		}
+			return FIoStatus::Ok;
+		};
 
-		if (UploadParams.bIgnoreContainerHeader)
-		{
-			UE_LOGF(LogIoStoreOnDemand, Display, "Ignoring container header for '%ls'", *Path)
-		}
-		else
+		FChunkValidationSettings ValidationSettings
 		{
-			const bool bDecrypt = false;
-			const FIoChunkId ChunkId = CreateContainerHeaderChunkId(ContainerFileReader.GetContainerId());
-			TIoStatusOr<FIoStoreCompressedReadResult> ReadResult = ContainerFileReader.ReadCompressed(ChunkId, FIoReadOptions(), bDecrypt);
-			if (ReadResult.IsOk())
-			{
-				const FIoBuffer& Chunk = ReadResult.ValueOrDie().IoBuffer;
-				ContainerEntry.Header = MakeArrayView<const uint8>(Chunk.GetData(), static_cast<int32>(Chunk.GetSize()));
-			}
-			else
-			{
-				const EIoErrorCode Err = ReadResult.Status().GetErrorCode();
-				if (Err != EIoErrorCode::UnknownChunkID && Err != EIoErrorCode::NotFound)
-				{
-					return ReadResult.Status();
-				}
-			}
-		}
-
+			.bValidate = UploadParams.bValidateChunks 
+		};
+
+		if (FIoStatus Status = UploadContainer(
+			Container,
+			EncryptionKeys,
+			UploadParams.MaxPartitionSize,
+			UploadParams.ContainerFlags,
+			ExistingChunks,
+			MoveTemp(Upload),
+			TocWriter,
+			ValidationSettings); !Status.IsOk())
 		{
-			const FString UTocFilePath = FPaths::ChangeExtension(Path, TEXT(".utoc"));
-			TArray<uint8> Buffer;
-			if (!FFileHelper::LoadFileToArray(Buffer, *UTocFilePath))
-			{
-				return FIoStatus(EIoErrorCode::ReadError, TEXT("Failed to upload .utoc file"));
-			}
-
-			ContainerEntry.UTocHash = FIoHash::HashBuffer(Buffer.GetData(), Buffer.Num());
-			TStringBuilder<256> Key;
-			if (UploadParams.BucketPrefix.IsEmpty() == false)
-			{
-				Key << UploadParams.BucketPrefix.ToLower() << TEXT("/");
-			}
-			Key << LexToString(ContainerEntry.UTocHash) << TEXT(".utoc");
-#if SHOULD_UPLOAD
-			const FS3PutObjectResponse Response = Client.TryPutObject(
-				FS3PutObjectRequest{UploadParams.Bucket, Key.ToString(), MakeMemoryView(Buffer.GetData(), Buffer.Num())});
-
-			if (!Response.IsOk())
-			{
-				return FIoStatus(EIoErrorCode::WriteError, FString::Printf(TEXT("Failed to upload '%s' (%s)"), *UTocFilePath, *Response.GetErrorStatus()));
-			}
-#endif
-			UTocPaths.Add(ContainerEntry.UTocHash, UTocFilePath);
-
-			UE_LOGF(LogIas, Display, "Uploaded '%ls'", *UTocFilePath);	
+			return Status;
 		}
 
 		if (UploadParams.bDeleteContainerFiles)
 		{
-			FilesToDelete.Add(Path);
-			ContainerFileReader.GetContainerFilePaths(FilesToDelete);
+			FilesToDelete.Add(Container.FilePath);
+			Container.Reader->GetContainerFilePaths(FilesToDelete);
 
-			// We need the pak files in order to mount OnDemand toc files!
-			if (UploadParams.bDeletePakFiles && !bWritePerContainerToc)
+			if (UploadParams.bDeletePakFiles)
 			{
-				FilesToDelete.Add(FPaths::ChangeExtension(Path, TEXT(".pak")));
-				FilesToDelete.Add(FPaths::ChangeExtension(Path, TEXT(".sig")));
+				FilesToDelete.Add(FPaths::ChangeExtension(Container.FilePath, TEXT(".pak")));
+				FilesToDelete.Add(FPaths::ChangeExtension(Container.FilePath, TEXT(".sig")));
 			}
 		}
 	}
 
-	if (OnDemandToc.Containers.IsEmpty())
-	{
-		return FIoStatus(EIoErrorCode::InvalidParameter, TEXT("No container file(s) marked as on demand"));
-	}
-#if SHOULD_UPLOAD
 	if (!UploadQueue.Flush())
 	{
 		return FIoStatus(EIoErrorCode::WriteError, TEXT("Failed to upload chunk(s)"));
 	}
-#endif
+
 	FString ServiceUrl;
-	GetTocPath(UploadParams.ServiceUrl, UploadParams.Bucket, UploadParams.BucketPrefix, ServiceUrl, OnDemandToc.Header.ChunksDirectory);
+	FString ChunksDirectory;
+	GetTocPath(UploadParams.ServiceUrl, UploadParams.Bucket, UploadParams.BucketPrefix, ServiceUrl, ChunksDirectory);
 
-	OnDemandToc.Meta.EpochTimestamp = FDateTime::Now().ToUnixTimestamp();
-	OnDemandToc.Meta.BuildVersion = UploadParams.BuildVersion;
-	OnDemandToc.Meta.TargetPlatform = UploadParams.TargetPlatform;
+	TocWriter.SetMetadata(UploadParams.BuildVersion, UploadParams.TargetPlatform);
+	TocWriter.SetChunksDirectory(ChunksDirectory);
+	TocWriter.SetHostGroup(UploadParams.HostGroupName);
 
-	FUploadResult UploadResult;
+	TStringBuilder<512> TocPath;
+	if (bGlobalToc)
 	{
-		using FOnDemandFileExt = UE::IoStore::Serialization::FOnDemandFileExt;
-
-		FLargeMemoryWriter Ar;
-		Ar << OnDemandToc;
-
-		UploadResult.CloudToc.Hash = FIoHash::HashBuffer(Ar.GetView());
-		const FString TocFilename = LexToString(UploadResult.CloudToc.Hash) + FOnDemandFileExt::CloudToc;
-
-		TStringBuilder<256> Key;
-		if (!UploadParams.BucketPrefix.IsEmpty())
-		{
-			Key << UploadParams.BucketPrefix.ToLower() << TEXT("/");
-		}
-		Key << TocFilename;
-
-		UploadResult.ServiceUrl = ServiceUrl;
-		UploadResult.CloudToc.Path = OnDemandToc.Header.ChunksDirectory / TocFilename;
-		UploadResult.CloudToc.Size = Ar.TotalSize();
-#if SHOULD_UPLOAD
-		const FS3PutObjectResponse Response = Client.TryPutObject(FS3PutObjectRequest{UploadParams.Bucket, Key.ToString(), Ar.GetView()});
-		if (Response.IsOk())
-		{
-			UE_LOGF(LogIas, Display, "Uploaded on demand TOC '%ls/%ls/%ls'", *Client.GetConfig().ServiceUrl, *UploadParams.Bucket, Key.ToString());
-		}
-		else
-		{
-			UE_LOGF(LogIas, Warning, "Failed to upload TOC '%ls/%ls/%ls' (%ls)", *Client.GetConfig().ServiceUrl, *UploadParams.Bucket, Key.ToString(), *Response.GetErrorStatus());
-			return FIoStatus(EIoErrorCode::WriteError, TEXT("Failed to upload TOC"));
-		}
-#endif
-		FLargeMemoryWriter FileAr;
-		if (UploadParams.bLegacyToc)
-		{
-			FileAr.Serialize(Ar.GetData(), Ar.TotalSize());
-		}
-		else
-		{
-			UE::IoStore::Serialization::FOnDemandTocWriter Writer;
-			if (TIoStatusOr<uint64> Status = Writer.Write(OnDemandToc, FileAr); !Status.IsOk())
-			{
-				return Status.Status();
-			}
-		}
-
-		UploadResult.FileToc.Hash = FIoHash::HashBuffer(FileAr.GetView());
-		UploadResult.FileToc.Path = OnDemandToc.Header.ChunksDirectory / LexToString(UploadResult.FileToc.Hash) + FOnDemandFileExt::CloudToc;
-		UploadResult.FileToc.Size = FileAr.TotalSize();
-
-		// Write a single global.uondemandtoc next to global.utoc or in the Cloud directory (legacy)
-		{
-			TStringBuilder<512> OnDemandTocFilePath;
-			if (UploadParams.GlobalContainerPath.IsEmpty() == false)
-			{
-				FPathViews::Append(OnDemandTocFilePath, UploadParams.GlobalContainerPath, FString(TEXT("global")) + FOnDemandFileExt::Toc);
-			}
-			else if (UploadParams.bWriteTocToDisk)
-			{
-				FPathViews::Append(OnDemandTocFilePath, UploadParams.TocOutputDir, LexToString(UploadResult.FileToc.Hash) + FOnDemandFileExt::CloudToc);
-			}
-
-			if (OnDemandTocFilePath.Len() > 0)
-			{
-				FPathViews::NormalizeFilename(OnDemandTocFilePath);
-				if (SaveArrayToFile(TArrayView64<uint8>(FileAr.GetData(), FileAr.TotalSize()), OnDemandTocFilePath.ToString()))
-				{
-					UE_LOGF(LogIas, Display, "Saved on-demand TOC file '%ls'", OnDemandTocFilePath.ToString());
-				}
-				else
-				{
-					UE_LOGF(LogIas, Error, "Failed to save on-demand TOC file '%ls'", OnDemandTocFilePath.ToString());
-				}
-			}
-		}
+		FPathViews::Append(TocPath, UploadParams.GlobalContainerPath, FString(TEXT("global")) + FOnDemandFileExt::Toc);
 	}
-
-	if (bWritePerContainerToc)
+	else
 	{
-		// Write out separate .uondemandtoc files, one per.utoc containing ondemand data.
-		FIoStatus Result = WriteContainerFiles(OnDemandToc, UTocPaths, UploadParams.bLegacyToc);
-		if (!Result.IsOk())
-		{
-			return Result;
-		}
+		TocPath << FPaths::ChangeExtension(Containers[0].FilePath, FOnDemandFileExt::Toc);
 	}
 
-	for (const FString& Path : FilesToDelete)
+	if (TIoStatusOr<uint64> Status = TocWriter.Write(TocPath.ToString()); Status.IsOk())
 	{
-		if (IFileManager::Get().FileExists(*Path))
-		{
-			UE_LOGF(LogIas, Display, "Deleting '%ls'", *Path); 
-			IFileManager::Get().Delete(*Path);
-		}
+		UE_LOGF(LogIoStoreOnDemand, Display, "Saved on-demand TOC file '%ls' (%.2lf KiB)",
+			TocPath.ToString(), double(Status.ConsumeValueOrDie()) / 1024.0);
 	}
+	else
+	{
+		return Status.Status();
+	}
+
+	GenerateTestChunks(Client, UploadParams.Bucket, ChunksRelativePath);
 
-	if (!UploadParams.TocOutputDir.IsEmpty())
+	TStringBuilder<256> BucketRelativePath;
+	if (!UploadParams.BucketPrefix.IsEmpty())
 	{
-		WriteConfigFile(UploadParams, UploadResult, KeyChain);
+		BucketRelativePath << UploadParams.BucketPrefix;
 	}
 
-	// Generate test chunks from 1KiB - 32MiB
+	for (const FContainerData& Container : Containers)
 	{
-		TStringBuilder<256> Key;
-		for (int32 Shft = 0; Shft < 16; ++Shft)
+		if (FIoStatus Status = UploadPartitionChunkHashes(
+			Client,
+			UploadParams.Bucket,
+			BucketRelativePath.ToString(),
+			UploadParams.TargetPlatform,
+			UploadParams.BuildVersion,
+			Container.UploadedPartitions); !Status.IsOk())
 		{
-			Key.Reset();
-			const int32 ChunkId		= (1 << Shft);
-			const int32 ChunkSize	= ChunkId << 10;
-			Key << ChunksRelativePath << TEXT("/") << ChunkId << TEXT("kib") << TEXT(".iochunk");
-
-			const FS3HeadObjectResponse HeadResponse = Client.HeadObject(FS3HeadObjectRequest{UploadParams.Bucket, Key.ToString()});
-			if (HeadResponse.IsOk())
-			{
-				continue;
-			}
+			return Status;
+		}
+	}
 
-			FIoBuffer Buffer(ChunkSize);
-			ANSICHAR* BufferContent = reinterpret_cast<ANSICHAR*>(Buffer.GetData());
-			FAnsiString SizeText = FAnsiString::Printf("%d", ChunkSize);
-			for (int32 Idx = 0; Idx < Buffer.GetSize(); ++Idx)
+	// Upload the original .utoc file(s)
+	{
+		for (const FContainerData& Container : Containers)
+		{
+			TStringBuilder<256> Key;
+			if (!UploadParams.BucketPrefix.IsEmpty())
 			{
-				BufferContent[Idx] = SizeText[Idx % SizeText.Len()];
+				Key << UploadParams.BucketPrefix << TEXT("/");
 			}
+			Key << Container.UTocHash << TEXT(".utoc");
 
-			const FS3PutObjectResponse Response = Client.TryPutObject(
-					FS3PutObjectRequest{UploadParams.Bucket, Key.ToString(), Buffer.GetView()}); 
+			const FS3PutObjectResponse Response = Client.TryPutObject(FS3PutObjectRequest{ UploadParams.Bucket, Key.ToString(), Container.UToc.GetView() });
 			if (Response.IsOk())
 			{
-				UE_LOGF(LogIas, Display, "Uploaded test chunk '%ls'", Key.ToString());
+				UE_LOGF(LogIoStoreOnDemand, Display, "Uploaded file TOC '%ls/%ls/%ls'",
+					*Client.GetConfig().ServiceUrl, *UploadParams.Bucket, Key.ToString());
 			}
 			else
 			{
-				UE_LOGF(LogIas, Display, "Failed to uploaded test chunk '%ls'", Key.ToString());
+				return FIoStatus(
+					EIoErrorCode::WriteError,
+					FString::Printf(TEXT("Failed to upload container file TOC, reason: %s"),
+						*Response.GetErrorStatus()));
 			}
 		}
 	}
 
+	// Delete the input file(s)
+	for (const FString& Path : FilesToDelete)
 	{
-		const double Duration = FPlatformTime::Seconds() - StartTime;
-		
-		UE_LOGF(LogIas, Display, "");
-		UE_LOGF(LogIas, Display, "------------------------------------------------- Upload Summary -------------------------------------------------");
-		UE_LOGF(LogIas, Display, "%-15ls: %ls", TEXT("Service URL"), *UploadParams.ServiceUrl);
-		UE_LOGF(LogIas, Display, "%-15ls: %ls", TEXT("Bucket"), *UploadParams.Bucket);
-		UE_LOGF(LogIas, Display, "%-15ls: %ls", TEXT("TargetPlatform"), *UploadParams.TargetPlatform);
-		UE_LOGF(LogIas, Display, "%-15ls: %ls", TEXT("BuildVersion"), *UploadParams.BuildVersion);
-		UE_LOGF(LogIas, Display, "%-15ls: %ls", TEXT("HostGroupName"), *UploadParams.HostGroupName);
-		UE_LOGF(LogIas, Display, "%-15ls: %ls", TEXT("TOC path"), *UploadResult.CloudToc.Path);
-		UE_LOGF(LogIas, Display, "%-15ls: %ls", TEXT("TOC flags"), *LexToString(UploadParams.TocFlags));
-		UE_LOGF(LogIas, Display, "%-15ls: %.2lf KiB", TEXT("TOC size"), double(UploadResult.CloudToc.Size) / 1024.0);
-		UE_LOGF(LogIas, Display, "%-15ls: %.2lf second(s)", TEXT("Duration"), Duration);
-		UE_LOGF(LogIas, Display, "");
-		
-		UE_LOGF(LogIas, Display, "%-50ls %15ls %15ls %15ls %15ls",
-			TEXT("Container"), TEXT("Chunk(s)"), TEXT("Size (MiB)"), TEXT("Uploaded"), TEXT("Uploaded (MiB)"));
-		UE_LOGF(LogIas, Display, "-------------------------------------------------------------------------------------------------------------------");
-
-		FContainerStats TotalStats;
-		for (const TTuple<FString, FContainerStats>& Kv : ContainerSummary)
+		if (IFileManager::Get().FileExists(*Path))
 		{
-			UE_LOGF(LogIas, Display, "%-50ls %15llu %15.2lf %15llu %15.2lf",
-				*Kv.Key, Kv.Value.ChunkCount, double(Kv.Value.TotalBytes) / 1024.0 / 1024.0, Kv.Value.UploadedChunkCount, double(Kv.Value.UploadedBytes) / 1024.0 / 1024.0);
-			
-			TotalStats.ChunkCount += Kv.Value.ChunkCount;
-			TotalStats.TotalBytes += Kv.Value.TotalBytes;
-			TotalStats.UploadedChunkCount += Kv.Value.UploadedChunkCount;
-			TotalStats.UploadedBytes += Kv.Value.UploadedBytes;
+			UE_LOGF(LogIoStoreOnDemand, Display, "Deleting '%ls'", *Path);
+			IFileManager::Get().Delete(*Path);
 		}
-		UE_LOGF(LogIas, Display, "-------------------------------------------------------------------------------------------------------------------");
-		UE_LOGF(LogIas, Display, "%-50ls %15llu %15.2lf %15llu %15.2lf",
-			TEXT("Total"),TotalStats.ChunkCount, double(TotalStats.TotalBytes) / 1024.0 / 1024.0, TotalStats.UploadedChunkCount, double(TotalStats.UploadedBytes) / 1024.0 / 1024.0);
-		UE_LOGF(LogIas, Display, "");
+	}
+
+	UE_LOGF(LogIoStoreOnDemand, Display, "");
+	UE_LOGF(LogIoStoreOnDemand, Display, "---------------------------------------------------- Upload Summary ----------------------------------------------------");
+	UE_LOGF(LogIoStoreOnDemand, Display, "%-20ls: %ls", TEXT("Service URL"), *UploadParams.ServiceUrl);
+	UE_LOGF(LogIoStoreOnDemand, Display, "%-20ls: %ls", TEXT("Bucket"), *UploadParams.Bucket);
+	UE_LOGF(LogIoStoreOnDemand, Display, "%-20ls: %ls", TEXT("BucketPrefix"), *UploadParams.BucketPrefix);
+	UE_LOGF(LogIoStoreOnDemand, Display, "%-20ls: %ls", TEXT("TargetPlatform"), *UploadParams.TargetPlatform);
+	UE_LOGF(LogIoStoreOnDemand, Display, "%-20ls: %ls", TEXT("BuildVersion"), *UploadParams.BuildVersion);
+	UE_LOGF(LogIoStoreOnDemand, Display, "%-20ls: %ls", TEXT("HostGroupName"), *UploadParams.HostGroupName);
+	UE_LOGF(LogIoStoreOnDemand, Display, "%-20ls: %ls", TEXT("Flags"), *LexToString(UploadParams.ContainerFlags));
+	UE_LOGF(LogIoStoreOnDemand, Display, "%-20ls: %.2lf MiB", TEXT("PartitionSize"), double(UploadParams.MaxPartitionSize) / 1024.0 / 1024.0);
+	UE_LOGF(LogIoStoreOnDemand, Display, "");
+
+	UE_LOGF(LogIoStoreOnDemand, Display, "%-40ls %15ls %15ls %15ls %15ls %15ls",
+		TEXT("Container"), TEXT("Size (MiB)"), TEXT("Chunk(s)"), TEXT("Partition(s)"), TEXT("Uploaded"), TEXT("Uploaded (MiB)"));
+	UE_LOGF(LogIoStoreOnDemand, Display, "------------------------------------------------------------------------------------------------------------------------");
+
+	FContainerStats TotalStats;
+	for (const FContainerData& Container : Containers)
+	{
+		const FContainerStats& Stats = Container.Stats;
+		UE_LOGF(LogIoStoreOnDemand, Display, "%-40ls %15.2lf %15llu %15llu %15llu %15.2lf",
+			*Container.Name, double(Stats.Chunks.Size) / 1024.0 / 1024.0, Stats.Chunks.Count,
+			Stats.Partitions.Count, Stats.UploadedPartitions.Count, double(Stats.UploadedPartitions.Size) / 1024.0 / 1024.0);
 		
-		UE_LOGF(LogIas, Display, "%-50ls %15ls %15ls %15ls", TEXT("Bucket"), TEXT("TOC(s)"), TEXT("Chunk(s)"), TEXT("MiB"));
-		UE_LOGF(LogIas, Display, "-------------------------------------------------------------------------------------------------------------------");
-		UE_LOGF(LogIas, Display, "%-50ls %15llu %15d %15.2lf", TEXT("Existing"), TotalExistingTocs, ExistingChunks.Num(), double(TotalExistingBytes) / 1024.0 / 1024.0);
-		UE_LOGF(LogIas, Display, "%-50ls %15llu %15d %15.2lf", TEXT("Uploaded"), 1, TotalStats.UploadedChunkCount, double(TotalStats.UploadedBytes) / 1024.0 / 1024.0);
-		UE_LOGF(LogIas, Display, "-------------------------------------------------------------------------------------------------------------------");
-		UE_LOGF(LogIas, Display, "%-50ls %15llu %15d %15.2lf", TEXT("Total"),
-			TotalExistingTocs + 1, ExistingChunks.Num() + TotalStats.UploadedChunkCount, double(TotalExistingBytes + TotalStats.UploadedBytes) / 1024.0 / 1024.0);
-		UE_LOGF(LogIas, Display, "");
+		TotalStats.Chunks				= TotalStats.Chunks + Stats.Chunks;
+		TotalStats.Partitions			= TotalStats.Partitions + Stats.Partitions;
+		TotalStats.UploadedPartitions	= TotalStats.UploadedPartitions + Stats.UploadedPartitions;
 	}
-	
-	return UploadResult;
+
+	UE_LOGF(LogIoStoreOnDemand, Display, "------------------------------------------------------------------------------------------------------------------------");
+	UE_LOGF(LogIoStoreOnDemand, Display, "%-40ls %15.2lf %15llu %15llu %15llu %15.2lf",
+		TEXT("Total"), double(TotalStats.Chunks.Size) / 1024.0 / 1024.0, TotalStats.Chunks.Count,
+		TotalStats.Partitions.Count, TotalStats.UploadedPartitions.Count,
+		double(TotalStats.UploadedPartitions.Size) / 1024.0 / 1024.0);
+	UE_LOGF(LogIoStoreOnDemand, Display, "");
+
+	return FIoStatus::Ok;
 }
 
 ////////////////////////////////////////////////////////////////////////////////
@@ -1186,10 +1219,10 @@ static int32 UploadCommandEntry(const FContext& Context)
 	FKeyChain KeyChain = Common::LoadCryptoKeys(Context);
 	FUploadParams Params = BuildUploadParams(Context);
 
-	TIoStatusOr<FUploadResult> Result = UploadContainerFiles(Params, Containers, KeyChain);
-	if (!Result.IsOk())
+	const FIoStatus Status = UploadContainerFiles(Params, Containers, KeyChain);
+	if (!Status.IsOk())
 	{
-		FString Reason = Result.Status().ToString();
+		FString Reason = Status.ToString();
 		Context.Abort(*Reason);
 	}
 
@@ -1197,8 +1230,6 @@ static int32 UploadCommandEntry(const FContext& Context)
 }
 
 ////////////////////////////////////////////////////////////////////////////////
-extern FArgumentSet S3Arguments;
-
 static FCommand UploadCommand(
 	UploadCommandEntry,
 	TEXT("Upload"),
@@ -1209,26 +1240,24 @@ static FCommand UploadCommand(
 		TArgument<FStringView>(TEXT("-BuildVersion"),			TEXT("Optional build version to embed it TOC")),
 		TArgument<FStringView>(TEXT("-TargetPlatform"),			TEXT("If given, embedded in the output TOC")),
 		TArgument<FStringView>(TEXT("-ConfigFilePath"),			TEXT("Path to the config file to write runtime parameters to")),
-		TArgument<FStringView>(TEXT("-DistributionUrl"),		TEXT("URL for IAS to use at runtime. Output to config file")),
-		TArgument<FStringView>(TEXT("-FallbackUrl"),			TEXT("Alternative fallack for DistributionUrl")),
 		TArgument<FStringView>(TEXT("-BucketPrefix"),			TEXT("Path to prefix to bucket objects")),
 		TArgument<FStringView>(TEXT("-HostGroupName"),			TEXT("Host group name or URL")),
 		TArgument<FStringView>(TEXT("-GlobalContainerPath"),	TEXT("Directory path to the content PAK folder")),
+		TArgument<FStringView>(TEXT("-TagSetFile"),				TEXT("Tag set JSON file.")),
 		TArgument<bool>(TEXT("-KeepContainerFiles"),			TEXT("Do not delete container files after upload")),
 		TArgument<bool>(TEXT("-KeepPakFiles"),					TEXT("Do not delete the springboard pak files")),
-		TArgument<bool>(TEXT("-WriteTocToDisk"),				TEXT("Output the TOC to disk as well as uploading")),
 		TArgument<bool>(TEXT("-PerContainerTocs"),				TEXT("Whether to generate TOC's for each container file(s)")),
 		TArgument<bool>(TEXT("-IgnoreContainerFlags"),			TEXT("Whether to ignore the OnDemand container flag")),
 		TArgument<bool>(TEXT("-StreamOnDemand"),				TEXT("Set the content to be streamed on-demand")),
 		TArgument<bool>(TEXT("-InstallOnDemand"),				TEXT("Set the content to be installed on-demand")),
-		TArgument<bool>(TEXT("-IgnoreContainerHeader"),			TEXT("Do not serialize the container header as part of the TOC file.")),
-		TArgument<bool>(TEXT("-LegacyToc"),						TEXT("Whether to serialize the TOC using the legacy format.")),
 		TArgument<int32>(TEXT("-MaxConcurrentUploads"),			TEXT("Number of simultaneous uploads")),
 		TArgument<int32>(TEXT("-MaxTocListCount"),				TEXT("Maximum number of TOC file(s) to list from the bucket")),
 		TArgument<int32>(TEXT("-MaxTocDownloadCount"),			TEXT("Maximum number of TOC file(s) to download")),
-		TArgument<FStringView>(TEXT("-TagSetFile"),		TEXT("File containing tags to add to files")),
+		TArgument<int32>(TEXT("-MaxPartitionSize"),				TEXT("Max partition size in bytes.")),
+		TArgument<bool>(TEXT("-ValidateChunks"),				TEXT("Runs additional validation on chunks after they are read from the container and before they are written out")),
 		S3Arguments,
 	}
 );
 
-} // namespace UE::IoStore::Tool
+} // namespace UE::IoStore::Tool::Upload
+
diff --git a/Engine/Source/Developer/IoStore/OnDemandUtilities/Private/Upload.h b/Engine/Source/Developer/IoStore/OnDemandUtilities/Private/Upload.h
new file mode 100644
--- /dev/null
+++ b/Engine/Source/Developer/IoStore/OnDemandUtilities/Private/Upload.h
@@ -0,0 +1,151 @@
+// Copyright Epic Games, Inc. All Rights Reserved.
+
+#pragma once
+
+#include "Containers/Array.h"
+#include "Containers/Map.h"
+#include "Containers/PagedArray.h"
+#include "Containers/UnrealString.h"
+#include "IO/IoBuffer.h"
+#include "IO/IoHash.h"
+#include "IO/IoStatus.h"
+#include "Misc/AES.h"
+#include "Templates/FunctionFwd.h"
+
+struct FIoStoreTocChunkInfo;
+class FIoStoreReader;
+
+namespace UE { class FS3Client; }
+namespace UE::IoStore::Serialization { class FOnDemandTocWriter; }
+namespace UE::IoStore::Serialization { enum class EOnDemandContainerEntryFlags : uint32; }
+
+namespace UE::IoStore::Tool::Upload
+{
+
+////////////////////////////////////////////////////////////////////////////////
+using FTagSets			= TMap<FString, TArray<FString>>;
+using FEncryptionKeys	= TMap<FGuid, FAES::FAESKey>;
+
+////////////////////////////////////////////////////////////////////////////////
+struct FUploadSettings
+{
+	static constexpr int32 MaxConcurrentUploads	= 16;
+	static constexpr int32 MinTocListCount		= 1000;
+	static constexpr int32 MaxTocListCount		= 100000;
+	static constexpr int32 MaxTocDownloadCount	= 10000;
+	static constexpr int32 DefaultPartitionSize = 2 << 20;
+	static constexpr int32 MaxPartitionSize		= 32 << 20;
+};
+
+////////////////////////////////////////////////////////////////////////////////
+struct FCountSize
+{
+	uint64 Count	= 0;
+	uint64 Size		= 0;
+};
+
+////////////////////////////////////////////////////////////////////////////////
+inline FCountSize operator+(const FCountSize& Lhs, const FCountSize& Rhs)
+{
+	return FCountSize
+	{
+		.Count = Lhs.Count + Rhs.Count,
+		.Size = Lhs.Size + Rhs.Size
+	};
+}
+
+////////////////////////////////////////////////////////////////////////////////
+struct FContainerStats
+{
+	FCountSize	Chunks;
+	FCountSize	Partitions;
+	FCountSize	UploadedPartitions;
+};
+
+////////////////////////////////////////////////////////////////////////////////
+struct FPartitionInfo
+{
+	FIoHash Hash = FIoHash::Zero;
+	uint64	Size = 0;
+};
+
+////////////////////////////////////////////////////////////////////////////////
+struct FContainerData
+{
+	FIoHash							UTocHash = FIoHash::Zero;
+	FString							FilePath;
+	FString							Name;
+	TUniquePtr<FIoStoreReader>		Reader;
+	FIoBuffer						Header;
+	FIoBuffer						UToc;
+	TArray<FIoStoreTocChunkInfo> 	ChunkInfos;
+	TMap<FString, TArray<uint32>>	TagSets;
+	TArray<FPartitionInfo>			UploadedPartitions;
+	FContainerStats					Stats;
+};
+
+////////////////////////////////////////////////////////////////////////////////
+using FUploadChunkCallback = TFunction<FIoStatus(FIoBuffer, const FIoHash&)>;
+
+////////////////////////////////////////////////////////////////////////////////
+struct FChunkValidationSettings
+{
+	bool	bValidate		= false;
+	uint64	MaxRawChunkSize = 32 << 20;
+};
+
+////////////////////////////////////////////////////////////////////////////////
+FIoStatus LoadContainers(
+	TConstArrayView<FString> FilePaths,
+	const FEncryptionKeys& EncryptionKeys,
+	TPagedArray<FContainerData>& OutContainers,
+	bool bIgnoreContainerFlags = false);
+
+////////////////////////////////////////////////////////////////////////////////
+FIoStatus LoadTagSets(
+	const FString& FilePath,
+	FTagSets& Out);
+
+////////////////////////////////////////////////////////////////////////////////
+FIoStatus CreateTagSets(
+	const FTagSets& TagSets,
+	FContainerData& Container);
+
+////////////////////////////////////////////////////////////////////////////////
+void GenerateTestChunks(
+	FS3Client& Client,
+	const FString& Bucket,
+	const FString& BucketRelativePath,
+	uint8 MaxExp = 16);
+
+////////////////////////////////////////////////////////////////////////////////
+FIoStatus DownloadPartitionChunkHashes(
+	FS3Client& Client,
+	const FString& Bucket,
+	const FString& BucketRelativePath,
+	int32 MaxTocListCount,
+	int32 MaxTocDownloadCount,
+	TSet<FIoHash>& Out);
+
+////////////////////////////////////////////////////////////////////////////////
+FIoStatus UploadPartitionChunkHashes(
+	FS3Client& Client,
+	const FString& Bucket,
+	const FString& BucketRelativePath,
+	const FString& TargetPlatform,
+	const FString& BuildVersion,
+	TConstArrayView<FPartitionInfo> Partitions);
+
+////////////////////////////////////////////////////////////////////////////////
+FIoStatus UploadContainer(
+	FContainerData& Container,
+	const FEncryptionKeys& EncryptionKeys,
+	int32 MaxPartitionSize,
+	UE::IoStore::Serialization::EOnDemandContainerEntryFlags ContainerFlags,
+	const TSet<FIoHash>& ExistingChunks,
+	FUploadChunkCallback&& Upload,
+	UE::IoStore::Serialization::FOnDemandTocWriter& OutWriter,
+	const FChunkValidationSettings& ValidationSettings = FChunkValidationSettings());
+
+} // namespace UE::IoStore::Tool::Upload
+
diff --git a/Engine/Source/Runtime/Experimental/IoStore/OnDemand/Private/OnDemandContentInstaller.cpp b/Engine/Source/Runtime/Experimental/IoStore/OnDemand/Private/OnDemandContentInstaller.cpp
--- a/Engine/Source/Runtime/Experimental/IoStore/OnDemand/Private/OnDemandContentInstaller.cpp
+++ b/Engine/Source/Runtime/Experimental/IoStore/OnDemand/Private/OnDemandContentInstaller.cpp
@@ -616,7 +616,7 @@ void FOnDemandContentInstaller::UpdateInstallRequestPriority(FSharedInternalInst
 				{
 					const FSharedOnDemandContainer& Container = Install.ResolvedChunks[PendingHttpRequest.ContainerIndex].Container;
 					const FOnDemandChunkEntry& ChunkEntry = Container->ChunkEntries[PendingHttpRequest.EntryIndex];
-					const FIoHash& ChunkHash = ChunkEntry.Hash;
+					const FOnDemandChunkHash& ChunkHash = ChunkEntry.Hash;
 
 					for (auto It = PendingChunkDownloads.CreateKeyIterator(ChunkHash); It; ++It)
 					{
@@ -1132,7 +1132,7 @@ void FOnDemandContentInstaller::ExecuteInstallRequest(FRequest& Request)
 		// Make sure we have enough space in the cache
 		TStaticArray<uint64, EOnDemandInstallCasType::Count> BytesToInstall(EInPlace::InPlace, 0);
 		{
-			TSet<FIoHash> ChunksToInstall;
+			TSet<FOnDemandChunkHash> ChunksToInstall;
 			ChunksToInstall.Reserve(InstallRequest.HttpRequestHandles.Num());
 			for (FChunkHttpRequestHandle& HttpRequest : InstallRequest.HttpRequestHandles)
 			{
@@ -1196,7 +1196,7 @@ void FOnDemandContentInstaller::ExecuteInstallRequest(FRequest& Request)
 		}
 		else
 		{
-			const FIoHash& ChunkHash = ChunkInfo.ChunkEntry().Hash;
+			const FOnDemandChunkHash& ChunkHash = ChunkInfo.ChunkEntry().Hash;
 
 			PendingChunkRequest = &PendingChunkDownloads.Add(ChunkHash);
 			PendingChunkRequest->RequestId = Request.SeqNo;
@@ -1205,11 +1205,11 @@ void FOnDemandContentInstaller::ExecuteInstallRequest(FRequest& Request)
 			// Setting HTTP flags to None to not use the HTTP cache
 			FIoHttpOptions HttpOptions = FIoHttpOptions(Request.Priority, HttpRetryCount, EIoHttpFlags::None);
 			HttpOptions.SetCategory(HttpCategory);
-			HttpOptions.SetRange(FIoHttpRange(0, ChunkInfo.Entry.GetDiskSize()));
+			HttpOptions.SetRange(FIoHttpRange(ChunkInfo.PartitionOffset(), ChunkInfo.PartitionOffset() + ChunkInfo.Entry.GetDiskSize()));
 			HttpOptions.SetMetadata(FIoHttpMetadata(ChunkId));
 
 			UE_LOGF(LogIoStoreOnDemand, VeryVerbose, "Created request for chunk: ChunkHash=%ls, RequestId=%u",
-				*TStringBuilder<64>(InPlace, ChunkHash), PendingChunkRequest->RequestId);
+				*LexToString(ChunkHash), PendingChunkRequest->RequestId);
 
 			TRACE_COUNTER_INCREMENT(OnDemandContentInstaller_HttpRequestsIssued);
 			PendingChunkRequest->HttpHandle = Batch.Get(
@@ -1217,7 +1217,7 @@ void FOnDemandContentInstaller::ExecuteInstallRequest(FRequest& Request)
 				ChunkInfo.RelativeUrl(),
 				FIoHttpHeaders(),
 				HttpOptions,
-				ChunkHash,
+				ChunkInfo.PartitionHash(),
 				[this, ChunkHash, RequestId = PendingChunkRequest->RequestId](FIoHttpResponse&& HttpResponse)
 				{
 					InstallerPipe.Launch(
@@ -1311,7 +1311,7 @@ void FOnDemandContentInstaller::ExectuteFlushLastAccessRequest(FRequest& Request
 		return;
 	}
 
-	TArray<FIoHash> TouchedChunks;
+	TArray<FOnDemandChunkHash> TouchedChunks;
 	for (const TPair<FSharedOnDemandContainer, TBitArray<>>& Pair : FlushLastAccessReqest.Args.RemovedChunkEntryReferences)
 	{
 		const FSharedOnDemandContainer& Container = Pair.Key;
@@ -1344,7 +1344,7 @@ void FOnDemandContentInstaller::ExectuteFlushLastAccessRequest(FRequest& Request
 }
 
 void FOnDemandContentInstaller::ProcessDownloadedChunk(
-	const FIoHash& ChunkHash, uint32 RequestId, EIoErrorCode InErrorCode, uint32 HttpStatusCode, FIoBuffer&& Chunk)
+	const FOnDemandChunkHash& ChunkHash, uint32 RequestId, EIoErrorCode InErrorCode, uint32 HttpStatusCode, FIoBuffer&& Chunk)
 {
 	TRACE_CPUPROFILER_EVENT_SCOPE(FOnDemandContentInstaller::ProcessDownloadedChunk);
 
@@ -1390,7 +1390,7 @@ void FOnDemandContentInstaller::ProcessDownloadedChunk(
 	}
 	else
 	{
-		const FIoHash VerifyChunkHash = FIoHash::HashBuffer(Chunk.GetView());
+		const FOnDemandChunkHash VerifyChunkHash = FOnDemandChunkHash::HashBuffer(Chunk.GetView());
 		if (ChunkHash == VerifyChunkHash)
 		{
 			TRACE_COUNTER_INCREMENT(OnDemandContentInstaller_HttpRequestsSuccess);
@@ -1416,7 +1416,7 @@ void FOnDemandContentInstaller::ProcessDownloadedChunk(
 				if (ensureMsgf(CVars::EnableConcurentInstallRequests(), TEXT("Install cache is full but iostore.ContentInstallerEnableConcurrentInstallRequests is not enabled!")))
 				{
 					uint64 PendingDownloadSize = 0;
-					for (const TPair<FIoHash, FChunkDownloadRequest>& Kv : PendingChunkDownloads)
+					for (const TPair<FOnDemandChunkHash, FChunkDownloadRequest>& Kv : PendingChunkDownloads)
 					{
 						const FChunkDownloadRequest& PendingChunkRequest = Kv.Value;
 						if (PendingChunkRequest.bChunkCanceled)
@@ -1468,9 +1468,9 @@ void FOnDemandContentInstaller::ProcessDownloadedChunk(
 			Result = MakeError(UnifiedError::IoStoreOnDemand::ChunkHashError::MakeError(
 				UnifiedError::IoStoreOnDemand::FChunkHashMismatchErrorContext
 				{
-					.ChunkId = FIoChunkId::InvalidChunkId, // Filled in later by OnProcessDownloadedChunkNotifyRequest
-					.ExpectedHash = ChunkHash,
-					.ActualHash = VerifyChunkHash
+					.ChunkId = FIoChunkId::InvalidChunkId,
+					.ExpectedHash = LexToString(ChunkHash),
+					.ActualHash = LexToString(VerifyChunkHash)
 				}));
 		}
 	}
@@ -1478,15 +1478,15 @@ void FOnDemandContentInstaller::ProcessDownloadedChunk(
 	static constexpr const TCHAR LogFmt[] = TEXT("ProcessDownloadedChunk: ChunkHash=%s, RequestId=%u, DownloadResult=%s");
 	if (ChunkRequest.bChunkCanceled)
 	{
-		UE_LOG(LogIoStoreOnDemand, Verbose, LogFmt, *TStringBuilder<64>(InPlace, ChunkHash), ChunkRequest.RequestId, TEXT("Cancelled"));
+		UE_LOG(LogIoStoreOnDemand, Verbose, LogFmt, *LexToString(ChunkHash), ChunkRequest.RequestId, TEXT("Cancelled"));
 	}
 	else if (Result.HasError())
 	{
-		UE_LOG(LogIoStoreOnDemand, Warning, LogFmt, *TStringBuilder<64>(InPlace, ChunkHash), ChunkRequest.RequestId, *Result.GetError().GetErrorMessage(true).ToString());
+		UE_LOG(LogIoStoreOnDemand, Warning, LogFmt, *LexToString(ChunkHash), ChunkRequest.RequestId, *Result.GetError().GetErrorMessage(true).ToString());
 	}
 	else
 	{
-		UE_LOG(LogIoStoreOnDemand, VeryVerbose, LogFmt, *TStringBuilder<64>(InPlace, ChunkHash), ChunkRequest.RequestId, TEXT("OK"));
+		UE_LOG(LogIoStoreOnDemand, VeryVerbose, LogFmt, *LexToString(ChunkHash), ChunkRequest.RequestId, TEXT("OK"));
 	}
 
 	while (FChunkHttpRequestHandle* ChunkHttpRequestHandle = ChunkRequest.ChunkRequestHttpHandles.PopHead())
@@ -1602,7 +1602,7 @@ void FOnDemandContentInstaller::TryCancelHttpRequestsForInstallRequest(FRequest:
 	{
 		const FSharedOnDemandContainer& PendingContainer = InstallRequest.ResolvedChunks[PendingHttpRequest.ContainerIndex].Container;
 		const FOnDemandChunkEntry& PendingChunkEntry = PendingContainer->ChunkEntries[PendingHttpRequest.EntryIndex];
-		const FIoHash& PendingChunkHash = PendingChunkEntry.Hash;
+		const FOnDemandChunkHash& PendingChunkHash = PendingChunkEntry.Hash;
 
 		FChunkDownloadRequest* ChunkDownload = nullptr;
 		for (auto It = PendingChunkDownloads.CreateKeyIterator(PendingChunkHash); It; ++It)
@@ -1632,7 +1632,7 @@ void FOnDemandContentInstaller::TryCancelHttpRequestsForInstallRequest(FRequest:
 				ChunkDownload->HttpHandle.Cancel();
 				ChunkDownload->bChunkCanceled = true;
 				++NumCancelled;
-				UE_LOGF(LogIoStoreOnDemand, Verbose, "Cancelled chunk download: hash - %ls, Id - %u", *TStringBuilder<64>(InPlace, PendingChunkHash), ChunkDownload->RequestId);
+				UE_LOGF(LogIoStoreOnDemand, Verbose, "Cancelled chunk download: hash - %ls, Id - %u", *LexToString(PendingChunkHash), ChunkDownload->RequestId);
 
 				// ProcessDownloadedChunk will remove the chunk request
 			}
diff --git a/Engine/Source/Runtime/Experimental/IoStore/OnDemand/Private/OnDemandContentInstaller.h b/Engine/Source/Runtime/Experimental/IoStore/OnDemand/Private/OnDemandContentInstaller.h
--- a/Engine/Source/Runtime/Experimental/IoStore/OnDemand/Private/OnDemandContentInstaller.h
+++ b/Engine/Source/Runtime/Experimental/IoStore/OnDemand/Private/OnDemandContentInstaller.h
@@ -267,7 +267,7 @@ private:
 	void							ExecuteVerifyRequest(FRequest& Request);
 	void							ExectuteFlushLastAccessRequest(FRequest& Request);
 	void							ProcessDownloadedChunk(
-										const FIoHash& ChunkHash, 
+										const FOnDemandChunkHash& Hash,
 										uint32 RequestId, 
 										EIoErrorCode InErrorCode, 
 										uint32 HttpStatusCode, 
@@ -293,7 +293,7 @@ private:
 	// This needs to be a multimap to support cancelation. At any time, there may be multiple
 	// requests for the same iochunk but only one of them should be active with the remainder
 	// waiting on final callbacks after being cancelled. See FChunkDownloadRequest.
-	TMultiMap<FIoHash, FChunkDownloadRequest> PendingChunkDownloads;
+	TMultiMap<FOnDemandChunkHash, FChunkDownloadRequest> PendingChunkDownloads;
 	const bool						bMMapCasEnabled{false};
 	std::atomic_bool				bShuttingDown{false};
 };
diff --git a/Engine/Source/Runtime/Experimental/IoStore/OnDemand/Private/OnDemandHttpIoDispacher.cpp b/Engine/Source/Runtime/Experimental/IoStore/OnDemand/Private/OnDemandHttpIoDispacher.cpp
--- a/Engine/Source/Runtime/Experimental/IoStore/OnDemand/Private/OnDemandHttpIoDispacher.cpp
+++ b/Engine/Source/Runtime/Experimental/IoStore/OnDemand/Private/OnDemandHttpIoDispacher.cpp
@@ -15,6 +15,7 @@
 #include "IO/HttpIoDispatcher.h"
 #include "IO/IoContainers.h"
 #include "IO/IoStoreOnDemand.h"
+#include "IO/Serialization/OnDemandContainerToc.h"
 #include "IasCache.h"
 #include "IasHostGroup.h"
 #include "Logging/LogVerbosity.h"
@@ -110,7 +111,6 @@ ENUM_CLASS_FLAGS(EHttpRequestTypeFilter);
 DECLARE_MULTICAST_DELEGATE(FOnRecreateHttpClient);
 FOnRecreateHttpClient OnRecreateHttpClient;
 
-
 ///////////////////////////////////////////////////////////////////////////////
 void OnHttpClientCVarChanged(IConsoleVariable* CVar)
 {
@@ -134,6 +134,15 @@ static FAutoConsoleVariableRef CVar_HttpIoDispatcherMaxRangeKiB(
 	ECVF_Default
 );
 
+int32 GHttpIoDispatcherMaxMergeGapKiB = 256;
+static FAutoConsoleVariableRef CVar_HttpIoDispatcherMaxMergeGapKiB(
+	TEXT("iax.HttpIoDispatcherMaxMergeGapKiB"),
+	GHttpIoDispatcherMaxMergeGapKiB,
+	TEXT("Maximum gap (in KiB) between two chunk ranges before they are no longer eligible to be merged into a single HTTP request.\n")
+	TEXT("	- a zero or negative value disables the gap limit (any gap is allowed up to iax.HttpIoDispatcherMaxRangeKiB)."),
+	ECVF_Default
+);
+
 int32 GIaxHttpSortRequestsByType = -1;
 static FAutoConsoleVariableRef CVar_IaxHttpSortRequestsByType(
 	TEXT("iax.HttpSortRequestsByType"),
@@ -357,6 +366,13 @@ static FAutoConsoleVariableRef CVar_UseDLBWManager(
 	TEXT("Toggle state for whether or not the HttpDispatcher will use the Downlink Bandwidth Manager.")
 );
 
+bool GLegacyBatchSorting = false;
+static FAutoConsoleVariableRef CVar_UseLegacyBatchSorting(
+	TEXT("iax.LegacyBatchSorting"),
+	GLegacyBatchSorting,
+	TEXT("Toggle between legacy or new batch sorting.")
+);
+
 ////////////////////////////////////////////////////////////////////////////////
 int32 CalculateRequestQueueLengthMax()
 {
@@ -713,7 +729,6 @@ static bool HttpRequestSortPredicate(const FHttpRequest& LHS, const FHttpRequest
 uint32 FHttpRequest::NextSeqNo = 0;
 
 ////////////////////////////////////////////////////////////////////////////////
-
 TArray<FIoOffsetAndLength> GetTotalRange(FHttpRequest* Request)
 {
 	TArray<FIoOffsetAndLength> Ranges;
@@ -759,6 +774,7 @@ TArray<FIoOffsetAndLength> GetTotalRange(FHttpRequest* Request)
 FHttpRequestList CreateBatch(
 	FHttpRequest& Request,
 	uint32 MaxRangeSize,
+	uint32 MaxGapSize,
 	FHttpRequestList&& List,
 	FHttpRequestList& OutRemaining)
 {
@@ -771,28 +787,143 @@ FHttpRequestList CreateBatch(
 		}
 	}
 
-	// The request is expected to be in the list
-	const bool bRemoved = List.Remove(&Request);
-	check(bRemoved);
+	if (GLegacyBatchSorting)
+	{
+		// The request is expected to be in the list
+		const bool bRemoved = List.Remove(&Request);
+		check(bRemoved);
+
+		FIoHttpRange		Range = Request.Options.GetRange();
+		FHttpRequestList	Batch(&Request);
 
-	FIoHttpRange		Range = Request.Options.GetRange();
-	FHttpRequestList	Batch(&Request);
+		while (FHttpRequest* NextRequest = List.PopHead())
+		{
+			FIoHttpRange NextRange = Range + NextRequest->Options.GetRange();
+			if (NextRange.GetSize() > MaxRangeSize)
+			{
+				OutRemaining.AddTail(NextRequest);
+			}
+			else
+			{
+				Range = NextRange;
+				Batch.AddTail(NextRequest);
+			}
+		}
 
-	while (FHttpRequest* NextRequest = List.PopHead())
+		return Batch;
+	}
+	else
 	{
-		FIoHttpRange NextRange = Range + NextRequest->Options.GetRange();
-		if (NextRange.GetSize() > MaxRangeSize)
+		auto AddRange = [](const FIoHttpRange& Lhs, const FIoHttpRange& Rhs, uint32& OutGap) -> FIoHttpRange
 		{
-			OutRemaining.AddTail(NextRequest);
+			const FIoHttpRange Sum	= Lhs + Rhs;
+			const uint32 Growth		= Sum.GetSize() - Lhs.GetSize();
+			const uint32 Useful		= Rhs.GetSize();
+			OutGap					= Growth > Useful ? Growth - Useful : 0;
+
+			return Sum;
+		};
+
+		OutRemaining = MoveTemp(List);
+
+		TArray<FHttpRequest*, TInlineAllocator<32>> SortedRequests;
+		for (FHttpRequest& R : OutRemaining)
+		{
+			SortedRequests.Add(&R);
 		}
-		else
+		SortedRequests.Sort([](const FHttpRequest& Lhs, const FHttpRequest& Rhs)
+		{
+			return Lhs.Options.GetRange().GetMin() < Rhs.Options.GetRange().GetMin();
+		});
+
+		const bool bRemoved = OutRemaining.Remove(&Request);
+		check(bRemoved);
+
+		FHttpRequestList	Batch(&Request);
+		FIoHttpRange		Range = Request.Options.GetRange();
+
+		const int32 Count	= SortedRequests.Num();
+		const int32 Index	= SortedRequests.Find(&Request);
+		int32 LhsIdx		= Index - 1;
+		int32 RhsIdx		= Index + 1;
+
+		for (;;)
 		{
+			FHttpRequest* Lhs		= LhsIdx > INDEX_NONE ? SortedRequests[LhsIdx] : nullptr;
+			FHttpRequest* Rhs		= RhsIdx < Count ? SortedRequests[RhsIdx] : nullptr;
+			FHttpRequest* Candidate	= nullptr;
+
+			if (Lhs == nullptr)
+			{
+				Candidate = Rhs;
+			}
+			else if (Rhs == nullptr)
+			{
+				Candidate = Lhs;
+			}
+			else
+			{
+				uint32 LhsGap = 0;
+				uint32 RhsGap = 0;
+				const FIoHttpRange LhsRange = AddRange(Range, Lhs->Options.GetRange(), LhsGap);
+				const FIoHttpRange RhsRange = AddRange(Range, Rhs->Options.GetRange(), RhsGap);
+
+				const bool bLhsValid = LhsRange.GetSize() <= MaxRangeSize && (MaxGapSize == MAX_uint32 || LhsGap <= MaxGapSize);
+				const bool bRhsValid = RhsRange.GetSize() <= MaxRangeSize && (MaxGapSize == MAX_uint32 || RhsGap <= MaxGapSize);
+
+				if (!bLhsValid && !bRhsValid)
+				{
+					Candidate = nullptr;
+				}
+				else if (!bLhsValid)
+				{
+					Candidate = Rhs;
+				}
+				else if (!bRhsValid)
+				{
+					Candidate = Lhs;
+				}
+				else
+				{
+					if (LhsGap != RhsGap)
+					{
+						Candidate = LhsGap < RhsGap ? Lhs : Rhs;
+					}
+					else
+					{
+						Candidate = LhsRange.GetSize() < RhsRange.GetSize() ? Lhs : Rhs;
+					}
+				}
+			}
+
+			if (Candidate == nullptr)
+			{
+				break;
+			}
+
+			uint32 Gap = 0;
+			const FIoHttpRange NextRange = AddRange(Range, Candidate->Options.GetRange(), Gap);
+			if ((NextRange.GetSize() > MaxRangeSize) || (MaxGapSize < MAX_uint32 && Gap > MaxGapSize))
+			{
+				break;
+			}
+
 			Range = NextRange;
-			Batch.AddTail(NextRequest);
+			OutRemaining.Remove(Candidate);
+			Batch.AddTail(Candidate);
+
+			if (Candidate == Lhs)
+			{
+				LhsIdx--;
+			}
+			else
+			{
+				RhsIdx++;
+			}
 		}
-	}
 
-	return Batch;
+		return Batch;
+	}
 }
 
 ////////////////////////////////////////////////////////////////////////////////
@@ -1162,7 +1293,8 @@ FHttpRequest* FHttpQueue::Dequeue()
 			FHttpRequestList	Requests(MoveTemp(*List));
 			*List								= FHttpRequestList();
 			const uint32		MaxRangeSize	= uint32(GHttpIoDispatcherMaxRangeKiB) << 10;
-			FHttpRequestList	Batch			= CreateBatch(*Next, MaxRangeSize, MoveTemp(Requests), *List);
+			const uint32		MaxGapSize		= GHttpIoDispatcherMaxMergeGapKiB > 0 ? uint32(GHttpIoDispatcherMaxMergeGapKiB) << 10 : MAX_uint32;
+			FHttpRequestList	Batch			= CreateBatch(*Next, MaxRangeSize, MaxGapSize, MoveTemp(Requests), *List);
 
 			if (List->IsEmpty())
 			{
@@ -2050,7 +2182,10 @@ bool FHttpDispatcher::ProcessHttpRequests()
 			{
 				// I/O store chunk request 
 				HashString << HttpRequest->ChunkHash;
-				Url << *HttpRequest->RelativeUrl << "/" << HashString.ToView().Left(2) << "/" << HashString << ANSITEXTVIEW(".iochunk");
+				Url << *HttpRequest->RelativeUrl
+					<< "/" << HashString.ToView().Left(2)
+					<< "/" << HashString
+					<< UE::IoStore::Serialization::FOnDemandFileExt::Partition; 
 				Ranges = GetTotalRange(HttpRequest);
 			}
 
diff --git a/Engine/Source/Runtime/Experimental/IoStore/OnDemand/Private/OnDemandHttpIoDispatcherBackend.cpp b/Engine/Source/Runtime/Experimental/IoStore/OnDemand/Private/OnDemandHttpIoDispatcherBackend.cpp
--- a/Engine/Source/Runtime/Experimental/IoStore/OnDemand/Private/OnDemandHttpIoDispatcherBackend.cpp
+++ b/Engine/Source/Runtime/Experimental/IoStore/OnDemand/Private/OnDemandHttpIoDispatcherBackend.cpp
@@ -42,10 +42,10 @@ static FAutoConsoleVariableRef CVar_IasFastCancelEnabled(
 );
 
 ////////////////////////////////////////////////////////////////////////////////
-static FIoHash GetHttpCacheKey(const FIoHash& ChunkHash, const FIoHttpRange& Range)
+static FIoHash GetHttpCacheKey(const FOnDemandChunkHash& ChunkHash, const FIoHttpRange& Range)
 {
 	FIoHashBuilder HashBuilder;
-	HashBuilder.Update(ChunkHash.GetBytes(), sizeof(FIoHash::ByteArray));
+	HashBuilder.Update(ChunkHash.GetData(), sizeof(FOnDemandChunkHash::ByteArray));
 	HashBuilder.Update(&Range, sizeof(FIoHttpRange));
 
 	return HashBuilder.Finalize();
@@ -225,11 +225,12 @@ void FHttpIoDispatcherBackend::ResolveIoRequests(FIoRequestList Requests, FIoReq
 				ResolvedSize).ConsumeValueOrDie();
 		}
 
-		FResolvedRequest*	ResolvedRequest = FResolvedRequest::Create(*Request, MoveTemp(ChunkInfo), ChunkRange);
-		const FIoHttpRange	HttpRange = FIoHttpRange::FromOffsetAndLength(ChunkRange);
-		const FIoHash		CacheKey = GetHttpCacheKey(ResolvedRequest->ChunkInfo.ChunkEntry().Hash, HttpRange);
-		const FIoHash&		ResourceHash = ResolvedRequest->ChunkInfo.ChunkEntry().Hash; // The content hash is the CDN resource hash
-		const EIoHttpFlags	Flags = EIoHttpFlags::Default; // Using default flags will fetch the request from the local cache
+		FResolvedRequest*			ResolvedRequest = FResolvedRequest::Create(*Request, MoveTemp(ChunkInfo), ChunkRange);
+		const FIoHash				CacheKey = GetHttpCacheKey(ResolvedRequest->ChunkInfo.ChunkEntry().Hash, FIoHttpRange::FromOffsetAndLength(ChunkRange));
+		const FIoOffsetAndLength	PartitionRange(ChunkInfo.PartitionOffset() + ChunkRange.GetOffset(), ChunkRange.GetLength());
+		const FIoHttpRange			HttpRange = FIoHttpRange::FromOffsetAndLength(PartitionRange);
+		const FIoHash&				ResourceHash = ResolvedRequest->ChunkInfo.PartitionHash(); // Hash of the CDN resource
+		const EIoHttpFlags			Flags = EIoHttpFlags::Default; // Default flags will fetch the request from the local cache
 
 		ResolvedRequest->HttpRequest = Batch.Get(
 			ResolvedRequest->ChunkInfo.HostGroupName(),
@@ -457,7 +458,7 @@ void FHttpIoDispatcherBackend::LogRequest(FResolvedRequest& ResolvedRequest, FIo
 		UncompressedSize >> 10,
 		DispacherRequest.Options.GetOffset(),
 		*LexToString(DispacherRequest.ChunkId),
-		*LexToString(ChunkInfo.Hash()),
+		*LexToString(ChunkInfo.PartitionHash()),
 		ChunkRange.GetOffset(), (ChunkRange.GetOffset() + ChunkRange.GetLength() - 1), ChunkSize,
 		100.0f * (float(ChunkRange.GetLength()) / float(ChunkSize)),
 		PrioToString(DispacherRequest.Priority));
diff --git a/Engine/Source/Runtime/Experimental/IoStore/OnDemand/Private/OnDemandInstallCache.cpp b/Engine/Source/Runtime/Experimental/IoStore/OnDemand/Private/OnDemandInstallCache.cpp
--- a/Engine/Source/Runtime/Experimental/IoStore/OnDemand/Private/OnDemandInstallCache.cpp
+++ b/Engine/Source/Runtime/Experimental/IoStore/OnDemand/Private/OnDemandInstallCache.cpp
@@ -23,7 +23,6 @@
 #include "IO/IoContainerHeader.h"
 #include "IO/IoChunkId.h"
 #include "IO/IoChunkEncoding.h"
-#include "IO/GenericHash.h"
 #include "Logging/StructuredLog.h"
 #include "IO/OnDemandError.h"
 #include "Misc/DateTime.h"
@@ -149,14 +148,6 @@ private:
 	EOnDemandInstallCasType CasType = EOnDemandInstallCasType::None;
 };
 
-
-///////////////////////////////////////////////////////////////////////////////
-using FCasAddr						= FHash96;
-static const FCasAddr& AsCasAddr(const FIoHash& IoHash)
-{
-	return *reinterpret_cast<const FCasAddr*>(&IoHash);
-}
-
 ///////////////////////////////////////////////////////////////////////////////
 struct FCasBlockId
 {
@@ -968,7 +959,7 @@ struct FCas
 
 	FResult						Configure(const FOnDemandInstallCasConfig& Config);
 	FResult						Initialize(FStringView Directory, bool bDeleteExisting = false);
-	FCasLocation				FindChunk(const FIoHash& Hash, const UE::FDeferLock* SkipLock = nullptr) const;
+	FCasLocation				FindChunk(const FCasAddr& ChunkAddr, const UE::FDeferLock* SkipLock = nullptr) const;
 	FCasBlockId					CreateBlock();
 	bool						ContainsBlock(FCasBlockId BlockId) const;
 	FResult						DeleteBlock(FCasBlockId BlockId, TFunctionRef<void(const FCasAddr&)> OutAddrs);
@@ -995,8 +986,8 @@ struct FCas
 	void						UpdateBlock(
 									FCasBlockId BlockId, 
 									uint32 BytesWritten, 
-									TConstArrayView<FIoHash> ChunkHashes, 
-									TConstArrayView<int64> Offsets,
+									TConstArrayView<FCasAddr> ChunkAddrs,
+									TConstArrayView<int64> ChunkOffsets,
 									TFunctionRef<void(const FCasAddr&, const FCasLocation&)> OutAddrs);
 	uint64						GetBlockInfo(FCasBlockInfoMap& OutBlockInfo) const;
 	uint64						GetTotalBlockSize() const;
@@ -1504,13 +1495,12 @@ FResult FCas::Initialize(FStringView Directory, bool bDeleteExisting)
 	return MakeValue();
 };
 
-FCasLocation FCas::FindChunk(const FIoHash& Hash, const UE::FDeferLock* SkipLock) const
+FCasLocation FCas::FindChunk(const FCasAddr& ChunkAddr, const UE::FDeferLock* SkipLock) const
 {
-	const FCasAddr& Addr	= AsCasAddr(Hash);
-	const uint32 TypeHash	= GetTypeHash(Addr);
+	const uint32 TypeHash	= GetTypeHash(ChunkAddr);
 	{
 		FMaybeLock Lock(this, SkipLock);
-		if (const FCasLocation* Loc = Lookup.FindByHash(TypeHash, Addr))
+		if (const FCasLocation* Loc = Lookup.FindByHash(TypeHash, ChunkAddr))
 		{
 			return *Loc;
 		}
@@ -2019,31 +2009,30 @@ bool FCas::UnlockedTrackModificationIf(ECasTrackModificationType Type, uint32 Bl
 void FCas::UpdateBlock(
 	FCasBlockId BlockId, 
 	uint32 BytesWritten, 
-	TConstArrayView<FIoHash> ChunkHashes, 
-	TConstArrayView<int64> Offsets,
+	TConstArrayView<FCasAddr> ChunkAddrs,
+	TConstArrayView<int64> ChunkOffsets,
 	TFunctionRef<void(const FCasAddr&, const FCasLocation&)> OutAddrs)
 {
-	check(ChunkHashes.Num() == Offsets.Num());
+	check(ChunkAddrs.Num() == ChunkOffsets.Num());
 	check(BlockId.IsValid());
 
 	TUniqueLock Lock(Mutex);
 
 	BlockIds.FindChecked(BlockId).Size += BytesWritten;
 
-	for (int32 Idx = 0, Count = Offsets.Num(); Idx < Count; ++Idx)
+	for (int32 Idx = 0, Count = ChunkOffsets.Num(); Idx < Count; ++Idx)
 	{
-		const FCasAddr&	CasAddr = AsCasAddr(ChunkHashes[Idx]);
-		const uint32	ChunkOffset = IntCastChecked<uint32>(Offsets[Idx]);
+		const FCasAddr&	ChunkAddr = ChunkAddrs[Idx];
+		const uint32	ChunkOffset = IntCastChecked<uint32>(ChunkOffsets[Idx]);
 
-		FCasLocation& Loc = Lookup.FindOrAdd(CasAddr);
+		FCasLocation& Loc = Lookup.FindOrAdd(ChunkAddr);
 		Loc.BlockId = BlockId;
 		Loc.BlockOffset = ChunkOffset;
 
-		OutAddrs(CasAddr, Loc);
+		OutAddrs(ChunkAddr, Loc);
 	}
 }
 
-
 uint64 FCas::GetBlockInfo(FCasBlockInfoMap& OutBlockInfo) const
 {
 	const int64	TimestampForMissingLastAccess = FCas::GetTimestampForMissingLastAccess();
@@ -2505,38 +2494,38 @@ class FOnDemandInstallCache final
 
 		bool IsEmpty() const
 		{
-			check(Chunks.Num() == ChunkHashes.Num());
-			return TotalSize == 0 && Chunks.IsEmpty() && ChunkHashes.IsEmpty();
+			check(Chunks.Num() == ChunkAddrs.Num());
+			return TotalSize == 0 && Chunks.IsEmpty() && ChunkAddrs.IsEmpty();
 		}
 
-		void Append(FIoBuffer&& Chunk, const FIoHash& ChunkHash)
+		void Append(FIoBuffer&& Chunk, const FCasAddr& ChunkAddr)
 		{
-			check(Chunks.Num() == ChunkHashes.Num());
+			check(Chunks.Num() == ChunkAddrs.Num());
 			TotalSize += Chunk.GetSize();
-			ChunkHashes.Add(ChunkHash);
+			ChunkAddrs.Add(ChunkAddr);
 			Chunks.Add(MoveTemp(Chunk));
 		}
 
 		void Reset()
 		{
 			Chunks.Reset();
-			ChunkHashes.Reset();
+			ChunkAddrs.Reset();
 			TotalSize = 0;
 		}
 
-		const FIoBuffer* FindChunk(const FIoHash& ChunkHash) const
+		const FIoBuffer* FindChunk(const FCasAddr& ChunkAddr) const
 		{
-			const int32 PendingIdx = ChunkHashes.IndexOfByKey(ChunkHash);
+			const int32 PendingIdx = ChunkAddrs.IndexOfByKey(ChunkAddr);
 			return (PendingIdx == INDEX_NONE) ? nullptr : &Chunks[PendingIdx];
 		}
 
-		bool ContainsChunk(const FIoHash& ChunkHash) const
+		bool ContainsChunk(const FCasAddr& ChunkAddr) const
 		{
-			return ChunkHashes.Contains(ChunkHash);
+			return ChunkAddrs.Contains(ChunkAddr);
 		}
 
 		TArray<FIoBuffer>		Chunks;
-		TArray<FIoHash>			ChunkHashes;
+		TArray<FCasAddr>		ChunkAddrs;
 		uint64					TotalSize = 0;
 		mutable FSharedMutex	SharedMutex;
 	};
@@ -2572,8 +2561,8 @@ public:
 	virtual const TCHAR*						GetName() const override;
 
 	// IOnDemandInstallCache
-	virtual bool								IsChunkCached(const FIoHash& ChunkHash) override;
-	virtual TResult<bool>						PutChunk(EOnDemandInstallCasType CasType, FIoBuffer&& Chunk, const FIoHash& ChunkHash) override;
+	virtual bool								IsChunkCached(const FCasAddr& ChunkAddr) override;
+	virtual TResult<bool>						PutChunk(EOnDemandInstallCasType CasType, FIoBuffer&& Chunk, const FCasAddr& ChunkAddr) override;
 	virtual FResult								Purge(EOnDemandInstallCasType CasType, uint64 BytesToInstall) override;
 	virtual FResult								PurgeAllUnreferenced(EOnDemandInstallCasType CasType, bool bDefrag, const uint64* BytesToPurge = nullptr) override;
 	virtual FResult								DefragAll(EOnDemandInstallCasType CasType, const uint64* BytesToFree = nullptr) override;
@@ -2581,7 +2570,7 @@ public:
 	virtual FResult								Flush(EOnDemandInstallCasType CasType) override;
 	virtual FResult								FlushLastAccess(const TStaticArray<bool, EOnDemandInstallCasType::Count>& bFlushCasLastAccess) override;
 	virtual void								UpdateLastAccess(
-													TConstArrayView<FIoHash> ChunkHashes, 
+													TConstArrayView<FCasAddr> ChunkAddrs,
 													TStaticArray<bool, EOnDemandInstallCasType::Count>& bInOutLastAccessDirty) override;
 	virtual FOnDemandInstallCacheUsage			GetCacheUsage() override;
 	virtual bool								IsCasEnabled(EOnDemandInstallCasType CasType) const override;
@@ -3406,7 +3395,7 @@ bool FOnDemandInstallCache::ResolveFromCas(FCacheState& CacheState, FIoRequestIm
 	return true;
 }
 
-bool FOnDemandInstallCache::IsChunkCached(const FIoHash& ChunkHash)
+bool FOnDemandInstallCache::IsChunkCached(const FCasAddr& ChunkAddr)
 {
 	for (FCacheState& CacheState : CacheStates)
 	{
@@ -3416,13 +3405,13 @@ bool FOnDemandInstallCache::IsChunkCached(const FIoHash& ChunkHash)
 		{
 			TSharedLock SharedLock(CacheState.PendingChunks.SharedMutex);
 
-			if (CacheState.PendingChunks.ContainsChunk(ChunkHash))
+			if (CacheState.PendingChunks.ContainsChunk(ChunkAddr))
 			{
 				return true;
 			}
 		}
 
-		const FCasLocation Loc = CacheState.Cas.FindChunk(ChunkHash);
+		const FCasLocation Loc = CacheState.Cas.FindChunk(ChunkAddr);
 		if (Loc.IsValid())
 		{
 			return true;
@@ -3432,7 +3421,7 @@ bool FOnDemandInstallCache::IsChunkCached(const FIoHash& ChunkHash)
 	return false;
 }
 
-TResult<bool> FOnDemandInstallCache::PutChunk(EOnDemandInstallCasType CasType, FIoBuffer&& Chunk, const FIoHash& ChunkHash)
+TResult<bool> FOnDemandInstallCache::PutChunk(EOnDemandInstallCasType CasType, FIoBuffer&& Chunk, const FCasAddr& ChunkAddr)
 {
 	FCacheState* CacheState = GetCacheState(CasType);
 	if (CacheState == nullptr)
@@ -3462,7 +3451,7 @@ TResult<bool> FOnDemandInstallCache::PutChunk(EOnDemandInstallCasType CasType, F
 	TUniqueLock Lock(CacheState->PendingChunks.SharedMutex);
 
 	FOnDemandInstallCacheStats::OnPutChunk(CasType, Chunk.GetSize());
-	CacheState->PendingChunks.Append(MoveTemp(Chunk), ChunkHash);
+	CacheState->PendingChunks.Append(MoveTemp(Chunk), ChunkAddr);
 	return MakeValue(true);
 }
 
@@ -3718,7 +3707,7 @@ FResult FOnDemandInstallCache::Verify()
 		Lookup.AddrToIndex.Reserve(Container->ChunkEntries.Num());
 		for (int32 EntryIndex = 0; const FOnDemandChunkEntry& Entry : Container->ChunkEntries)
 		{
-			const FCasAddr& Addr = AsCasAddr(Entry.Hash);
+			const FCasAddr& Addr = Entry.Hash;
 			Lookup.AddrToIndex.Add(Addr, EntryIndex++);
 		}
 	}
@@ -3827,7 +3816,7 @@ FResult FOnDemandInstallCache::Verify()
 				continue;
 			}
 
-			const FIoHash ChunkHash = FIoHash::HashBuffer(Chunk.GetView().Left(ChunkSize));
+			const FOnDemandChunkHash ChunkHash = FOnDemandChunkHash::HashBuffer(Chunk.GetView().Left(ChunkSize));
 
 			if (ChunkHash == ChunkEntry.Hash)
 			{
@@ -4041,7 +4030,7 @@ void FOnDemandInstallCache::AddReferencesToBlocks(
 	OutTotalReferencedBytes = 0;
 
 	// TODO: this could use the CAS address to be smaller
-	TSet<FIoHash> VisitedReferencedChunks;
+	TSet<FCasAddr> VisitedReferencedChunks;
 	{
 		int32 ReserveSize = 0;
 		for (int32 Index = 0; FSharedOnDemandContainer Container : Containers)
@@ -4108,8 +4097,6 @@ void FOnDemandInstallCache::AddReferencesToBlocks(
 			EntryIndex++;
 		}
 	}
-
-	return;
 }
 
 FResult FOnDemandInstallCache::PurgeInternal(
@@ -4211,10 +4198,10 @@ FResult FOnDemandInstallCache::Defrag(
 
 	struct FDefragBlockReferencedChunk
 	{
-		FIoHash		Hash;
-		FIoChunkId	ChunkId; 
-		uint32		BlockOffset = 0;
-		uint32		DiskSize = 0;
+		FOnDemandChunkHash	Hash;
+		FIoChunkId			ChunkId; 
+		uint32				BlockOffset = 0;
+		uint32				DiskSize = 0;
 	};
 
 	struct FDefragBlock
@@ -4528,7 +4515,7 @@ FResult FOnDemandInstallCache::Defrag(
 					return Result;
 				}
 
-				const FIoHash ChunkHash = FIoHash::HashBuffer(Buffer.GetView());
+				const FOnDemandChunkHash ChunkHash = FOnDemandChunkHash::HashBuffer(Buffer.GetView());
 				if (ChunkHash != ReffedChunk.Hash)
 				{
 					UE_LOGF(LogIoStoreOnDemand, Error, "Found chunk with invalid hash while defragging block, ChunkId='%ls', BlockId=%u, BlockOffset=%u",
@@ -4552,10 +4539,8 @@ FResult FOnDemandInstallCache::Defrag(
 					FResult Result = MakeError(InstallCacheDefragError::MakeError(
 						FChunkHashMismatchErrorContext
 						{
-							.ChunkId		= ReffedChunk.ChunkId, // May not be unique, just return the first found for debugging purposes
-							.ExpectedHash	= ReffedChunk.Hash,
-							.ActualHash		= ChunkHash,
-							.CasType		= CacheState.Cas.GetType()
+							.ChunkId = ReffedChunk.ChunkId, // May not be unique, just return the first found for debugging purposes
+							.CasType = CacheState.Cas.GetType()
 						}));
 					FOnDemandInstallCacheStats::OnDefrag(CacheState.Cas.GetType(), Result, TotalFragmentedBytes, TotalCachedBytes, BlocksToDefragFragmentedSize, BlocksToDefragTotalSize);
 					return Result;
@@ -4682,9 +4667,9 @@ FResult FOnDemandInstallCache::FlushLastAccess(const TStaticArray<bool, EOnDeman
 	return MakeValue();
 }
 
-void FOnDemandInstallCache::UpdateLastAccess(TConstArrayView<FIoHash> ChunkHashes, TStaticArray<bool, EOnDemandInstallCasType::Count>& bInOutLastAccessDirty)
+void FOnDemandInstallCache::UpdateLastAccess(TConstArrayView<FCasAddr> ChunkAddrs, TStaticArray<bool, EOnDemandInstallCasType::Count>& bInOutLastAccessDirty)
 {
-	if (ChunkHashes.IsEmpty())
+	if (ChunkAddrs.IsEmpty())
 	{
 		return;
 	}
@@ -4697,10 +4682,8 @@ void FOnDemandInstallCache::UpdateLastAccess(TConstArrayView<FIoHash> ChunkHashe
 		const EOnDemandInstallCasType CasType = CacheState.Cas.GetType();
 		TUniqueLock Lock(CacheState.Cas);
 
-		for (const FIoHash& ChunkHash : ChunkHashes)
+		for (const FCasAddr& Addr : ChunkAddrs)
 		{
-			const FCasAddr& Addr = AsCasAddr(ChunkHash);
-
 			if (FCasLocation* CasLoc = CacheState.Cas.Lookup.Find(Addr))
 			{
 				const FCasBlockId BlockId = CasLoc->BlockId;
@@ -4864,7 +4847,7 @@ FResult FOnDemandInstallCache::FlushPendingChunksImpl(FCacheState& CacheState, c
 		const int64 CasBlockOffset = CasFileHandle->Tell();
 		
 		TArray<int64>		Offsets;
-		const FIoHash*		ChunkHashBegin = &Chunks.ChunkHashes[ChunkIdx];
+		const FCasAddr*		ChunkAddrBegin = &Chunks.ChunkAddrs[ChunkIdx];
 
 		while (ChunkIdx < Chunks.Chunks.Num())
 		{
@@ -4875,9 +4858,9 @@ FResult FOnDemandInstallCache::FlushPendingChunksImpl(FCacheState& CacheState, c
 
 			if (!ensure(Chunks.Chunks[ChunkIdx].GetSize() <= CacheState.Cas.GetMaxBlockSize()))
 			{
-				UE_LOGFMT(LogIoStoreOnDemand, Error, "{CasType} CAS - Chunk is larger then block size ({ChunkId})", 
+				UE_LOGFMT(LogIoStoreOnDemand, Error, "{CasType} CAS - Chunk is larger then block size ({ChunkAddr})",
 					("CasType", LexToString(CacheState.Cas.GetType())), 
-					("ChunkId", Chunks.ChunkHashes[ChunkIdx]));
+					("ChunkAddr", LexToString(Chunks.ChunkAddrs[ChunkIdx])));
 			}
 
 			const FIoBuffer& Chunk = Chunks.Chunks[ChunkIdx];
@@ -4887,7 +4870,7 @@ FResult FOnDemandInstallCache::FlushPendingChunksImpl(FCacheState& CacheState, c
 			++ChunkIdx;
 		}
 
-		TConstArrayView<FIoHash> ChunkHashes = MakeArrayView(ChunkHashBegin, Offsets.Num());
+		TConstArrayView<FCasAddr> ChunkAddrs = MakeArrayView(ChunkAddrBegin, Offsets.Num());
 		const int64 BytesToWrite = Ar.Tell();
 
 		if (BytesToWrite > 0)
@@ -4951,7 +4934,7 @@ FResult FOnDemandInstallCache::FlushPendingChunksImpl(FCacheState& CacheState, c
 				Transaction.BlockModification(CurrentBlockId, Now);
 			}
 
-			CacheState.Cas.UpdateBlock(CurrentBlockId, IntCastChecked<uint32>(BytesToWrite), ChunkHashes, Offsets,
+			CacheState.Cas.UpdateBlock(CurrentBlockId, IntCastChecked<uint32>(BytesToWrite), ChunkAddrs, Offsets,
 				[&Transaction](const FCasAddr& CasAddr, const FCasLocation& Loc)
 				{
 					Transaction.ChunkLocation(Loc, CasAddr);
diff --git a/Engine/Source/Runtime/Experimental/IoStore/OnDemand/Private/OnDemandInstallCache.h b/Engine/Source/Runtime/Experimental/IoStore/OnDemand/Private/OnDemandInstallCache.h
--- a/Engine/Source/Runtime/Experimental/IoStore/OnDemand/Private/OnDemandInstallCache.h
+++ b/Engine/Source/Runtime/Experimental/IoStore/OnDemand/Private/OnDemandInstallCache.h
@@ -5,6 +5,7 @@
 #include "IO/IoDispatcherBackend.h"
 #include "IO/IoStoreOnDemand.h"
 #include "IO/OnDemandError.h"
+#include "IO/GenericHash.h"
 #include "Templates/SharedPointer.h"
 
 struct FIoContainerHeader; 
@@ -15,7 +16,9 @@ namespace UE::IoStore
 class FOnDemandContentHandle;
 class FOnDemandIoStore;
 struct FOnDemandInstallCacheUsage;
-using FSharedOnDemandContainer = TSharedPtr<struct FOnDemandContainer, ESPMode::ThreadSafe>;
+
+using FSharedOnDemandContainer	= TSharedPtr<struct FOnDemandContainer, ESPMode::ThreadSafe>;
+using FCasAddr					= FHash96;
 
 struct FOnDemandInstallCacheStorageUsage
 {
@@ -29,15 +32,15 @@ class IOnDemandInstallCache
 {
 public:
 	virtual										~IOnDemandInstallCache() = default;
-	virtual bool								IsChunkCached(const FIoHash& ChunkHash) = 0;
-	virtual TResult<bool>						PutChunk(EOnDemandInstallCasType CasType, FIoBuffer&& Chunk, const FIoHash& ChunkHash) = 0;
+	virtual bool								IsChunkCached(const FCasAddr& CasAddr) = 0;
+	virtual TResult<bool>						PutChunk(EOnDemandInstallCasType CasType, FIoBuffer&& Chunk, const FCasAddr& CasAddr) = 0;
 	virtual FResult								Purge(EOnDemandInstallCasType CasType, uint64 BytesToInstall) = 0;
 	virtual FResult								PurgeAllUnreferenced(EOnDemandInstallCasType CasType, bool bDefrag, const uint64* BytesToPurge = nullptr) = 0;
 	virtual FResult								DefragAll(EOnDemandInstallCasType CasType, const uint64* BytesToFree = nullptr) = 0;
 	virtual FResult								Verify() = 0;
 	virtual FResult								Flush(EOnDemandInstallCasType CasType) = 0;
 	virtual FResult								FlushLastAccess(const TStaticArray<bool, EOnDemandInstallCasType::Count>& bFlushCasLastAccess) = 0;
-	virtual void								UpdateLastAccess(TConstArrayView<FIoHash> ChunkHashes, TStaticArray<bool, EOnDemandInstallCasType::Count>& bInOutLastAccessDirty) = 0;
+	virtual void								UpdateLastAccess(TConstArrayView<FCasAddr> CasAddrs, TStaticArray<bool, EOnDemandInstallCasType::Count>& bInOutLastAccessDirty) = 0;
 	virtual FOnDemandInstallCacheUsage			GetCacheUsage() = 0;
 	virtual bool								IsCasEnabled(EOnDemandInstallCasType CasType) const = 0;
 };
diff --git a/Engine/Source/Runtime/Experimental/IoStore/OnDemand/Private/OnDemandIoStore.cpp b/Engine/Source/Runtime/Experimental/IoStore/OnDemand/Private/OnDemandIoStore.cpp
--- a/Engine/Source/Runtime/Experimental/IoStore/OnDemand/Private/OnDemandIoStore.cpp
+++ b/Engine/Source/Runtime/Experimental/IoStore/OnDemand/Private/OnDemandIoStore.cpp
@@ -6,6 +6,7 @@
 #include "Algo/AnyOf.h"
 #include "Algo/RemoveIf.h"
 #include "Algo/Transform.h"
+#include "Algo/IsSorted.h"
 #include "Async/UniqueLock.h"
 #include "Containers/StringConv.h"
 #include "Containers/Ticker.h"
@@ -1085,8 +1086,8 @@ FIoStatus FOnDemandIoStore::InitializeInstallCache()
 					Value = Seed;
 				}
 
-				const FIoHash ChunkHash = FIoHash::HashBuffer(Chunk.GetView());
-				TResult<bool> PutResult = InstallCache->PutChunk(EOnDemandInstallCasType::General, MoveTemp(Chunk), ChunkHash);
+				const FOnDemandChunkHash ChunkAddr = FOnDemandChunkHash::HashBuffer(Chunk.GetView());
+				TResult<bool> PutResult = InstallCache->PutChunk(EOnDemandInstallCasType::General, MoveTemp(Chunk), ChunkAddr);
 				if (PutResult.HasError())
 				{
 					Result = MakeError(PutResult.StealError());
@@ -1444,7 +1445,8 @@ FIoStatus FOnDemandIoStore::CreateContainersFromToc(
 	TConstArrayView<FIoContainerId> Existing,
 	TArray<FSharedOnDemandContainer>& Out)
 {
-	using namespace UE::IoStore::Serialization::V1;
+	using namespace UE::IoStore::Serialization;
+	using namespace UE::IoStore::Serialization::V2;
 
 	check(!MountArgs.MountId.IsEmpty());
 	check(!MountArgs.FilePath.IsEmpty());
@@ -1512,6 +1514,7 @@ FIoStatus FOnDemandIoStore::CreateContainersFromToc(
 		FOnDemandContainerTocView View	= MaybeView.ConsumeValueOrDie();
 		Container->HeaderView			= View.Header;
 		Container->ContainerHeaderChunk	= MoveTemp(View.ContainerHeaderChunk);
+		Container->PartitionEntries		= View.PartitionEntries;
 		Container->ChunkIds				= View.ChunkIds;
 		Container->ChunkEntries			= View.ChunkEntries;
 		Container->BlockSizes			= View.BlockSizes;
@@ -1535,6 +1538,8 @@ FIoStatus FOnDemandIoStore::CreateContainersFromToc(
 			EnumAddFlags(Container->Flags, EOnDemandContainerFlags::StreamOnDemand);
 		}
 
+		check(Algo::IsSorted(Container->ChunkIds));
+
 		// Do not count output container memory as being allocated for utocs
 		LLM_TAGSET_SCOPE_CLEAR(ELLMTagSet::Assets);
 		LLM_TAGSET_SCOPE_CLEAR(ELLMTagSet::AssetClasses);
diff --git a/Engine/Source/Runtime/Experimental/IoStore/OnDemand/Private/OnDemandIoStore.h b/Engine/Source/Runtime/Experimental/IoStore/OnDemand/Private/OnDemandIoStore.h
--- a/Engine/Source/Runtime/Experimental/IoStore/OnDemand/Private/OnDemandIoStore.h
+++ b/Engine/Source/Runtime/Experimental/IoStore/OnDemand/Private/OnDemandIoStore.h
@@ -58,10 +58,13 @@ enum class EOnDemandContainerFlags : uint8
 ENUM_CLASS_FLAGS(EOnDemandContainerFlags);
 
 ///////////////////////////////////////////////////////////////////////////////
-using FOnDemandChunkEntry				= UE::IoStore::Serialization::V1::FOnDemandChunkEntry;
-using FOnDemandTagSet					= UE::IoStore::Serialization::V1::FOnDemandTagSetEntry;
+using FOnDemandPartitionEntry			= UE::IoStore::Serialization::V2::FOnDemandPartitionEntry;
+using FOnDemandChunkEntry				= UE::IoStore::Serialization::V2::FOnDemandChunkEntry;
+using FOnDemandTagSet					= UE::IoStore::Serialization::V2::FOnDemandTagSetEntry;
 using FOnDemandTocStorage				= UE::IoStore::Serialization::FOnDemandTocStorage;
-using FOnDemandContainerTocHeaderView	= UE::IoStore::Serialization::V1::FOnDemandContainerTocHeaderView;
+using FOnDemandContainerTocHeaderView	= UE::IoStore::Serialization::V2::FOnDemandContainerTocHeaderView;
+using FOnDemandPartitionHash			= UE::IoStore::Serialization::V2::FOnDemandPartitionHash;
+using FOnDemandChunkHash				= UE::IoStore::Serialization::V2::FOnDemandChunkHash;
 
 ///////////////////////////////////////////////////////////////////////////////
 struct FOnDemandChunkEntryReferences
@@ -105,6 +108,7 @@ struct FOnDemandContainer
 	FSharedContainerHeader					Header;
 	FIASHostGroup							HostGroup;
 	FString									MountId;
+	TConstArrayView<FOnDemandPartitionEntry> PartitionEntries;
 	TConstArrayView<FIoChunkId>				ChunkIds;
 	TConstArrayView<FOnDemandChunkEntry> 	ChunkEntries;
 	TConstArrayView<uint32>					BlockSizes;
@@ -209,12 +213,13 @@ FAnsiString FOnDemandContainer::GetTestUrl()
 	}
 
 	TAnsiStringBuilder<41> HashString;
-	HashString << ChunkEntries[0].Hash;
+	HashString << PartitionEntries[ChunkEntries[0].PartitionIndex].Hash;
 
 	TAnsiStringBuilder<256> Url;
-	Url << "/" << RelativeUrl.GetView()
+	Url << RelativeUrl.GetView()
 		<< "/" << HashString.ToView().Left(2)
-		<< "/" << HashString << ANSITEXTVIEW(".iochunk");
+		<< "/" << HashString
+		<< UE::IoStore::Serialization::FOnDemandFileExt::Partition;
 
 	return Url.ToString();
 }
@@ -230,7 +235,9 @@ struct FOnDemandChunkInfo
 												: SharedContainer(MoveTemp(Other.SharedContainer))
 												, Entry(Other.Entry) { }
 
-	const									FIoHash& Hash() const { return Entry.Hash; }
+	const FOnDemandChunkHash&				Hash() const { return Entry.Hash; }
+	const FIoHash&							PartitionHash() const { return SharedContainer->PartitionEntries[Entry.PartitionIndex].Hash; }
+	uint32									PartitionOffset() const { return Entry.PartitionOffset; }
 	uint32									RawSize() const { return Entry.RawSize; }
 	uint32									EncodedSize() const { return Entry.EncodedSize; }
 	uint32									DiskSize() const { return Entry.GetDiskSize(); }
@@ -279,7 +286,7 @@ TConstArrayView<FIoBlockHash> FOnDemandChunkInfo::BlockHashes() const
 		? TConstArrayView<FIoBlockHash>()
 		: Entry.BlockInfo.Count() > 1
 			? TConstArrayView<FIoBlockHash>(SharedContainer->BlockHashes.GetData() + Entry.BlockInfo.Offset(), Entry.BlockInfo.Count())
-			: TConstArrayView<uint32>(Entry.BlockInfo.Data() + 1, 1);
+			: TConstArrayView<FIoBlockHash>(Entry.BlockInfo.Data() + 1, 1);
 }
 
 const FIASHostGroup& FOnDemandChunkInfo::HostGroup() const
diff --git a/Engine/Source/Runtime/Experimental/IoStore/OnDemandCore/Internal/IO/Serialization/OnDemandContainerToc.h b/Engine/Source/Runtime/Experimental/IoStore/OnDemandCore/Internal/IO/Serialization/OnDemandContainerToc.h
--- a/Engine/Source/Runtime/Experimental/IoStore/OnDemandCore/Internal/IO/Serialization/OnDemandContainerToc.h
+++ b/Engine/Source/Runtime/Experimental/IoStore/OnDemandCore/Internal/IO/Serialization/OnDemandContainerToc.h
@@ -6,6 +6,7 @@
 #include "Containers/StringFwd.h"
 #include "IO/IoBuffer.h"
 #include "IO/IoContainerId.h"
+#include "IO/GenericHash.h"
 #include "IO/IoHash.h"
 #include "IO/IoStatus.h"
 #include "Misc/AES.h"
@@ -32,9 +33,10 @@ namespace UE::IoStore::Serialization
 ///////////////////////////////////////////////////////////////////////////////
 struct FOnDemandFileExt
 {
-	static constexpr const TCHAR* Toc		= TEXT(".uondemandtoc");
-	static constexpr const TCHAR* CloudToc	= TEXT(".iochunktoc");
-	static constexpr const TCHAR* Chunk		= TEXT(".iochunk");
+	static constexpr const TCHAR* Toc			= TEXT(".uondemandtoc");
+	static constexpr const TCHAR* Chunk			= TEXT(".iochunk");
+	static constexpr const TCHAR* PartitionToc	= TEXT(".ioparttoc");
+	static constexpr const TCHAR* Partition		= TEXT(".iopart");
 };
 
 ///////////////////////////////////////////////////////////////////////////////
@@ -71,6 +73,7 @@ enum class EOnDemandTocMinorVersion : uint16
 {
 	Invalid			= 0,
 	MemoryMapped	= 1,
+	Partitions		= 2,
 
 	LatestPlusOne,
 	Latest			= (LatestPlusOne - 1)
@@ -138,6 +141,17 @@ private:
 	FStorage	ContainerData;
 };
 
+////////////////////////////////////////////////////////////////////////////////
+enum class EOnDemandContainerEntryFlags : uint32
+{
+	None			= 0,
+	InstallOnDemand	= (1 << 0),
+	StreamOnDemand	= (1 << 1),
+
+	Last			= StreamOnDemand
+};
+ENUM_CLASS_FLAGS(EOnDemandContainerEntryFlags);
+
 ////////////////////////////////////////////////////////////////////////////////
 class FOnDemandTocWriter
 {
@@ -148,8 +162,35 @@ public:
 
 	UE_API FOnDemandTocWriter& operator=(FOnDemandTocWriter&&);
 
-	UE_API TIoStatusOr<uint64> Write(const UE::IoStore::FOnDemandToc& Toc, FArchive& Ar);
-	UE_API TIoStatusOr<uint64> Write(const UE::IoStore::FOnDemandToc& Toc, const FString& Filename);
+	UE_API bool					IsEmpty() const;
+	UE_API void					SetMetadata(
+									const FString& BuildVersion,
+									const FString& TargetPlatform);
+	UE_API void					SetHostGroup(const FString& HostGroupName);
+	UE_API void					SetChunksDirectory(const FString& ChunksDirectory);
+	UE_API void					BeginContainer(
+									const FString& ContainerName,
+									const FIoContainerId& ContainerId,
+									const FIoHash& UTocHash,
+									const FGuid& EncryptionKeyGuid,
+									uint32 CompressionBlockSize,
+									uint32 FileContainerFlags,
+									EOnDemandContainerEntryFlags ContainerFlags,
+									FIoBuffer ContainerHeader);
+	UE_API int32				BeginPartition();
+	UE_API uint64				AddChunk(
+									int32 Partition,
+									const FIoChunkId& ChunkId,
+									FIoBuffer EncodedChunk,
+									TConstArrayView<uint32> BlockSizes,
+									TConstArrayView<FIoBlockHash> BlockHashes,
+									uint64 RawChunkSize);
+	UE_API void					AddTagSet(FStringView Tag, TConstArrayView<uint32> PackageIndices);
+	UE_API FIoBuffer			EndPartition(int32 Partition, FIoHash& OutHash);
+	UE_API void 				EndContainer(FName CompressionFormat);
+
+	UE_API TIoStatusOr<uint64>	Write(FArchive& Ar);
+	UE_API TIoStatusOr<uint64>	Write(const FString& Filename);
 
 private:
 	FOnDemandTocWriter(const FOnDemandTocWriter&) = delete;
@@ -187,14 +228,6 @@ struct FOnDemandTocHeader
 };
 static_assert(sizeof(FOnDemandTocHeader) == 128);
 
-enum class EOnDemandContainerEntryFlags : uint32
-{
-	None			= 0,
-	InstallOnDemand	= (1 << 0),
-	StreamOnDemand	= (1 << 1),
-};
-ENUM_CLASS_FLAGS(EOnDemandContainerEntryFlags);
-
 struct FOnDemandContainerEntry
 {
 	FGuid					EncryptionKeyGuid;
@@ -211,7 +244,8 @@ struct FOnDemandContainerEntry
 	FIoHash					UTocHash = FIoHash::Zero;
 	uint32					ContainerFlags = 0; // EOnDemandContainerEntryFlags
 	uint32					FileContainerFlags = 0; // EIoContainerFlags
-	int8					Pad[36] = {0};
+	uint32					PartitionCount = 0;
+	int8					Pad[32] = {0};
 };
 static_assert(sizeof(FOnDemandContainerEntry) == 128);
 
@@ -363,6 +397,99 @@ private:
 
 } // V1
 
+////////////////////////////////////////////////////////////////////////////////
+namespace V2
+{
+
+////////////////////////////////////////////////////////////////////////////////
+using FOnDemandStringEntry				= V1::FOnDemandStringEntry;
+using FOnDemandTocHeader				= V1::FOnDemandTocHeader;
+using FOnDemandContainerEntry			= V1::FOnDemandContainerEntry;
+using FOnDemandChunkBlockInfo			= V1::FOnDemandChunkBlockInfo;
+using FOnDemandTagSetEntry				= V1::FOnDemandTagSetEntry;
+using FOnDemandContainerTocHeaderView	= V1::FOnDemandContainerTocHeaderView;
+using EOnDemandTocReaderOptions			= V1::EOnDemandTocReaderOptions;
+using FOnDemandChunkHash				= FHash96;
+using FOnDemandPartitionHash			= FIoHash;
+
+////////////////////////////////////////////////////////////////////////////////
+struct FOnDemandPartitionEntry
+{
+	FIoHash	Hash = FIoHash::Zero;
+	uint32	Size = 0;
+};
+static_assert(sizeof(FOnDemandPartitionEntry) == 24);
+
+////////////////////////////////////////////////////////////////////////////////
+struct FOnDemandChunkEntry
+{
+	uint32						GetDiskSize() const { return Align(EncodedSize, FAES::AESBlockSize); }
+
+	FOnDemandChunkHash			Hash;
+	uint32						PartitionIndex = 0;
+	uint32						PartitionOffset = 0;
+	uint32						RawSize = 0;
+	uint32						EncodedSize = 0;
+	FOnDemandChunkBlockInfo		BlockInfo;
+};
+static_assert(sizeof(FOnDemandChunkEntry) == 36);
+
+///////////////////////////////////////////////////////////////////////////////
+struct FOnDemandContainerTocView
+{
+	FOnDemandContainerTocHeaderView				Header;
+	TConstArrayView<FOnDemandPartitionEntry>	PartitionEntries;
+	TConstArrayView<FIoChunkId>					ChunkIds;
+	TConstArrayView<FOnDemandChunkEntry>		ChunkEntries;
+	TConstArrayView<uint32>						BlockSizes;
+	TConstArrayView<FIoBlockHash>				BlockHashes;
+	TConstArrayView<FOnDemandTagSetEntry>		TagSets;
+	TConstArrayView<uint32>						TagSetIndices;
+	FIoBuffer									ContainerHeaderChunk; // FIoContainerHeader
+};
+
+////////////////////////////////////////////////////////////////////////////////
+class FOnDemandTocReader
+{
+public:
+	FOnDemandTocReader() = default;
+	FOnDemandTocReader(FOnDemandTocReader&&) = default;
+	FOnDemandTocReader& operator=(FOnDemandTocReader&&) = default;
+
+	FDateTime				CreationTime() const { return FDateTime::FromUnixTimestamp(Header->EpochTimestamp); }
+	FUtf8StringView			BuildVersion() const { return GetString(Header->BuildVersion); }
+	FUtf8StringView			TargetPlatform() const { return GetString(Header->TargetPlatform); }
+	FUtf8StringView			ChunksDirectory() const { return GetString(Header->ChunksDirectory); }
+	FUtf8StringView			HostGroupName() const { return GetString(Header->HostGroupName); }
+	FUtf8StringView			CompressionFormat() const { return GetString(Header->CompressionFormat); }
+	UE_API FUtf8StringView	GetString(FOnDemandStringEntry StringEntry) const;
+
+	TConstArrayView<FOnDemandContainerEntry> Containers() const { return ContainerEntries; }
+
+	UE_API TIoStatusOr<FOnDemandContainerTocView> ReadContainer(
+		const FOnDemandContainerEntry& Container,
+		FOnDemandTocStorage& Storage,
+		EOnDemandTocReaderOptions Options) const;
+	UE_API static TIoStatusOr<FOnDemandTocReader> Read(const FString& Filename);
+
+private:
+	FOnDemandTocReader(const FString& Filename, FIoBuffer Buffer);
+	FOnDemandTocReader(const FOnDemandTocReader&) = delete;
+	FOnDemandTocReader& operator=(const FOnDemandTocReader&) = delete;
+
+	FString										Filename;
+	FIoBuffer									Buffer;
+	const FOnDemandTocHeader*					Header = nullptr;
+	const UTF8CHAR*								StringTable = nullptr;
+	TConstArrayView<FOnDemandContainerEntry>	ContainerEntries;
+};
+
+} // namespace V2
+	
 } // UE::IoStore::Serialization
 
+////////////////////////////////////////////////////////////////////////////////
+UE_API FStringBuilderBase& operator<<(FStringBuilderBase& Sb, UE::IoStore::Serialization::EOnDemandContainerEntryFlags Flags);
+UE_API FString LexToString(UE::IoStore::Serialization::EOnDemandContainerEntryFlags Flags);
+
 #undef UE_API
diff --git a/Engine/Source/Runtime/Experimental/IoStore/OnDemandCore/Private/Serialization/OnDemandContainerToc.cpp b/Engine/Source/Runtime/Experimental/IoStore/OnDemandCore/Private/Serialization/OnDemandContainerToc.cpp
--- a/Engine/Source/Runtime/Experimental/IoStore/OnDemandCore/Private/Serialization/OnDemandContainerToc.cpp
+++ b/Engine/Source/Runtime/Experimental/IoStore/OnDemandCore/Private/Serialization/OnDemandContainerToc.cpp
@@ -20,21 +20,17 @@
 #include "Templates/UniquePtr.h"
 #include "Templates/UnrealTemplate.h"
 
-#ifndef WITH_IOSTORE_ONDEMAND_TESTS
-#define WITH_IOSTORE_ONDEMAND_TESTS 0
-#endif
-
-#if WITH_IOSTORE_ONDEMAND_TESTS
-#include "TestHarness.h"
-#include "TestMacros/Assertions.h"
-#include <catch2/generators/catch_generators.hpp>
-#include "IO/IoChunkId.h"
-#include "IO/PackageId.h"
-#endif
-
 namespace UE::IoStore::Serialization
 {
 
+////////////////////////////////////////////////////////////////////////////////
+template <typename T>
+FArchive& SerializeArrayData(FArchive& Ar, const TArray<T>& Array)
+{
+	Ar.Serialize((void*)Array.GetData(), Array.GetTypeSize() * Array.Num());
+	return Ar;
+}
+
 ////////////////////////////////////////////////////////////////////////////////
 FOnDemandTocStorage::FOnDemandTocStorage()
 {
@@ -109,129 +105,275 @@ FOnDemandTocStorageType FOnDemandTocStorage::StorageType() const
 ////////////////////////////////////////////////////////////////////////////////
 class FOnDemandTocWriter::FImpl
 {
-	using FOnDemandStringEntry = V1::FOnDemandStringEntry;
+	using FOnDemandStringEntry		= V1::FOnDemandStringEntry;
+	using FOnDemandChunkEntry		= V2::FOnDemandChunkEntry;
+	using FOnDemandPartitionEntry	= V2::FOnDemandPartitionEntry;
+	using FOnDemandTagSetEntry		= V1::FOnDemandTagSetEntry;
+
+	struct FContainer
+	{
+		FIoHash								UTocHash;
+		FGuid								EncryptionKeyGuid;
+		TUniquePtr<FLargeMemoryWriter>		PartitionAr;
+		TArray<FOnDemandPartitionEntry>		Partitions;
+		TArray<FIoChunkId>					ChunkIds;
+		TArray<FOnDemandChunkEntry>			ChunkEntries;
+		TArray<uint32>						BlockSizes;
+		TArray<FIoBlockHash>				BlockHashes;
+		TArray<FOnDemandTagSetEntry>		TagSets;
+		TArray<uint32>						TagSetIndices;
+		FIoBuffer							Header;
+		FString								Name;
+		FIoContainerId						Id;
+		uint32								CompressionBlockSize;
+		int32								CurrentPartition = INDEX_NONE;
+		FName								CompressionFormat = NAME_None;
+		uint32								FileContainerFlags = 0;
+		EOnDemandContainerEntryFlags		ContainerFlags = EOnDemandContainerEntryFlags::None;
+	};
+
+	struct FMetadata
+	{
+		FString BuildVersion;
+		FString TargetPlatform;
+		FString HostGroupName;
+		FString ChunksDirectory;
+	};
 
 public:
-	TIoStatusOr<uint64> Write(const UE::IoStore::FOnDemandToc& Toc, FArchive& Ar)
+	bool IsEmpty() const
+	{
+		return Containers.IsEmpty();
+	}
+
+	void SetMetadata(const FString& BuildVersion, const FString& TargetPlatform)
 	{
-		using namespace UE::IoStore::Serialization::V1;
+		Metadata.BuildVersion	= BuildVersion;
+		Metadata.TargetPlatform = TargetPlatform;
+	}
+
+	void SetHostGroup(const FString& HostGroupName)
+	{
+		Metadata.HostGroupName = HostGroupName;
+	}
+
+	void SetChunksDirectory(const FString& ChunksDirectory)
+	{
+		Metadata.ChunksDirectory = ChunksDirectory;
+	}
+
+	void BeginContainer(
+		const FString& ContainerName,
+		const FIoContainerId& ContainerId,
+		const FIoHash& UTocHash,
+		const FGuid& EncryptionKeyGuid,
+		uint32 CompressionBlockSize,
+		uint32 FileContainerFlags,
+		EOnDemandContainerEntryFlags ContainerFlags,
+		FIoBuffer ContainerHeader)
+	{
+		check(CurrentContainer == nullptr);
+		const int32 Idx				= Containers.Add();
+		CurrentContainer			= &Containers[Idx];
+
+		CurrentContainer->UTocHash				= UTocHash;
+		CurrentContainer->EncryptionKeyGuid		= EncryptionKeyGuid;
+		CurrentContainer->Header				= ContainerHeader;
+		CurrentContainer->Name					= ContainerName;
+		CurrentContainer->Id					= ContainerId;
+		CurrentContainer->CompressionBlockSize	= CompressionBlockSize;
+		CurrentContainer->FileContainerFlags	= FileContainerFlags;
+		CurrentContainer->ContainerFlags		= ContainerFlags;
+	}
+
+	int32 BeginPartition()
+	{
+		check(CurrentContainer != nullptr);
+		check(CurrentContainer->PartitionAr.IsValid() == false);
+		check(CurrentContainer->CurrentPartition == INDEX_NONE);
+		CurrentContainer->CurrentPartition = CurrentContainer->Partitions.Num();
+		CurrentContainer->Partitions.AddDefaulted();
+		CurrentContainer->PartitionAr = MakeUnique<FLargeMemoryWriter>();
+
+		return CurrentContainer->CurrentPartition;
+	}
+
+	uint64 AddChunk(
+		int32 Partition,
+		const FIoChunkId& ChunkId,
+		FIoBuffer EncodedChunk,
+		TConstArrayView<uint32> BlockSizes,
+		TConstArrayView<FIoBlockHash> BlockHashes,
+		uint64 RawChunkSize)
+	{
+		using namespace V2;
+		check(CurrentContainer != nullptr);
+		check(CurrentContainer->PartitionAr.IsValid());
+		check(BlockSizes.Num() == BlockHashes.Num());
+		check(ChunkId.IsValid());
+
+		FArchive& Ar = *CurrentContainer->PartitionAr;
+
+		CurrentContainer->ChunkIds.Add(ChunkId);
+
+		FOnDemandChunkEntry& ChunkEntry = CurrentContainer->ChunkEntries.AddDefaulted_GetRef();
+		ChunkEntry.Hash					= FOnDemandChunkHash::HashBuffer(EncodedChunk.GetView());
+		ChunkEntry.PartitionIndex		= uint32(CurrentContainer->CurrentPartition);
+		ChunkEntry.PartitionOffset		= uint32(Ar.Tell());
+		ChunkEntry.RawSize				= uint32(RawChunkSize);
+		ChunkEntry.EncodedSize			= uint32(EncodedChunk.GetSize());
+
+		// If the chunk only has a single compression block, the size and block hash are encoded into the chunk entry
+		const int32 BlockCount = BlockSizes.Num();
+		if (BlockCount == 1)
+		{
+			ChunkEntry.BlockInfo = FOnDemandChunkBlockInfo::FromSizeAndHash(BlockSizes[0], BlockHashes[0]);
+		}
+		else
+		{
+			const int32 BlocksOffset = CurrentContainer->BlockSizes.Num();
+			ChunkEntry.BlockInfo = FOnDemandChunkBlockInfo::FromOffsetAndCount(BlocksOffset, BlockCount);
+			CurrentContainer->BlockSizes.Append(BlockSizes);
+			CurrentContainer->BlockHashes.Append(BlockHashes);
+		}
+
+		Ar.Serialize(EncodedChunk.GetData(), EncodedChunk.GetSize());
+		return uint64(Ar.TotalSize());
+	}
+
+	void AddTagSet(FStringView Tag, TConstArrayView<uint32> PackageIndices)
+	{
+		check(CurrentContainer != nullptr);
+		check(!Tag.IsEmpty());
+		check(!PackageIndices.IsEmpty());
+
+		FOnDemandTagSetEntry& Entry = CurrentContainer->TagSets.AddDefaulted_GetRef();
+		Entry.Tag					= WriteString(Tag);
+		Entry.Count					= uint32(PackageIndices.Num());
+		Entry.Offset				= uint32(CurrentContainer->TagSetIndices.Num());
+
+		CurrentContainer->TagSetIndices.Append(PackageIndices);
+	}
+
+	FIoBuffer EndPartition(int32 Partition, FIoHash& OutHash)
+	{
+		check(CurrentContainer != nullptr);
+		check(CurrentContainer->PartitionAr.IsValid());
+
+		FLargeMemoryWriter& Ar		= *CurrentContainer->PartitionAr;
+		const uint64 PartitionSize	= Ar.TotalSize();
+
+		FOnDemandPartitionEntry& PartitionEntry = CurrentContainer->Partitions[Partition];
+		FIoBuffer OutPartition(FIoBuffer::AssumeOwnership, Ar.ReleaseOwnership(), PartitionSize);
+
+		PartitionEntry.Hash	= FIoHash::HashBuffer(OutPartition.GetView());
+		PartitionEntry.Size	= uint32(OutPartition.GetSize());
+		OutHash				= PartitionEntry.Hash;
 
+		CurrentContainer->PartitionAr.Reset();
+		CurrentContainer->CurrentPartition = INDEX_NONE;
+
+		return OutPartition;
+	}
+
+	void EndContainer(FName CompressionFormat)
+	{
+		check(CurrentContainer != nullptr);
+		CurrentContainer->CompressionFormat = CompressionFormat;
+		CurrentContainer = nullptr;
+	}
+
+	TIoStatusOr<uint64> Write(FArchive& Ar)
+	{
+		using namespace UE::IoStore::Serialization::V2;
+
+		if (Containers.IsEmpty())
+		{
+			return FIoStatusBuilder(EIoErrorCode::FileOpenFailed)
+				<< TEXT("No container(s) to serialize");
+		}
+
+		FName							CompressionFormat = NAME_None;
 		TArray<FOnDemandContainerEntry> ContainerEntries;
 		FLargeMemoryWriter				ContainerDataAr;
 		const int64						ArStart = Ar.Tell();
 
-		for (int32 ContainerIndex = 0; const UE::IoStore::FOnDemandTocContainerEntry& LegacyContainer : Toc.Containers)
+		for (int32 ContainerIndex = 0; const FContainer& Container : Containers)
 		{
-			TConstArrayView<uint32>			LegacyContainerBlockSizes = LegacyContainer.BlockSizes;
-			TConstArrayView<FIoBlockHash>	LegacyContainerBlockHashes = LegacyContainer.BlockHashes;
-			TArray<FIoChunkId>				ChunkIds;
-			TArray<FOnDemandChunkEntry>		ChunkEntries;
-			TArray<uint32>					BlockSizes;
-			TArray<FIoBlockHash>			BlockHashes;
-			TArray<FOnDemandTagSetEntry>	TagSets;
-			TArray<uint32>					TagSetIndices;
+			check(Container.ChunkIds.Num() == Container.ChunkEntries.Num());
 
 			TArray<int32> EntryIndices;
-			EntryIndices.Reserve(LegacyContainer.Entries.Num());
-			for (int32 Idx = 0; Idx < LegacyContainer.Entries.Num(); ++Idx)
+			EntryIndices.Reserve(Container.ChunkIds.Num());
+			for (int32 Idx = 0; Idx < Container.ChunkIds.Num(); ++Idx)
 			{
 				EntryIndices.Add(Idx);
 			}
 
-			EntryIndices.Sort([&LegacyContainer](int32 LHS, int32 RHS)
+			EntryIndices.Sort([&Container](int32 LHS, int32 RHS)
 			{
-				return LegacyContainer.Entries[LHS].ChunkId < LegacyContainer.Entries[RHS].ChunkId;
+				return Container.ChunkIds[LHS] < Container.ChunkIds[RHS];
 			});
 
-			for (int32 EntryIndex : EntryIndices)
-			{
-				const FOnDemandTocEntry&		LegacyEntry = LegacyContainer.Entries[EntryIndex];
-				TConstArrayView<uint32>			LegacyBlockSizes = LegacyContainerBlockSizes.Mid(LegacyEntry.BlockOffset, LegacyEntry.BlockCount);
-				TConstArrayView<FIoBlockHash>	LegacyBlockHashes = LegacyContainerBlockHashes.Mid(LegacyEntry.BlockOffset, LegacyEntry.BlockCount);
-
-				ChunkIds.Add(LegacyEntry.ChunkId);
-
-				FOnDemandChunkBlockInfo ChunkBlockInfo =
-					FOnDemandChunkBlockInfo::FromOffsetAndCount(uint32(BlockSizes.Num()), LegacyEntry.BlockCount);
-
-				// If the chunk only has a single compression block, the size and block hash are encoded into the chunk entry
-				if (LegacyBlockSizes.Num() > 1)
-				{
-					BlockSizes.Append(LegacyBlockSizes);
-					BlockHashes.Append(LegacyBlockHashes);
-				}
-				else
-				{
-					check(!LegacyBlockSizes.IsEmpty() && !LegacyBlockHashes.IsEmpty());
-					ChunkBlockInfo = FOnDemandChunkBlockInfo::FromSizeAndHash(LegacyBlockSizes[0], LegacyBlockHashes[0]);
-				}
-
-				ChunkEntries.Add(FOnDemandChunkEntry
-				{
-					.Hash			= LegacyEntry.Hash,
-					.RawSize		= IntCastChecked<uint32>(LegacyEntry.RawSize),
-					.EncodedSize	= IntCastChecked<uint32>(LegacyEntry.EncodedSize),
-					.BlockInfo		= ChunkBlockInfo
-				});
-			}
+			TArray<FIoChunkId>			SortedChunkIds;
+			TArray<FOnDemandChunkEntry>	SortedChunkEntries;
+			const TArray<uint32>&		BlockSizes = Container.BlockSizes;
+			const TArray<FIoBlockHash>&	BlockHashes = Container.BlockHashes;
 
-			for (const FOnDemandTocTagSet& LegacyTagSet : Toc.TagSets)
+			SortedChunkIds.Reserve(Container.ChunkIds.Num());
+			SortedChunkEntries.Reserve(Container.ChunkIds.Num());
+			for (int32 Idx : EntryIndices)
 			{
-				for (const FOnDemandTocTagSetPackageList& PkgList : LegacyTagSet.Packages)
-				{
-					if (PkgList.ContainerIndex == ContainerIndex)
-					{
-						FOnDemandTagSetEntry& TagSetEntry	= TagSets.AddDefaulted_GetRef();
-						TagSetEntry.Tag						= WriteString(LegacyTagSet.Tag);
-						TagSetEntry.Offset					= TagSetIndices.Num();
-						TagSetEntry.Count					= PkgList.PackageIndicies.Num();
-
-						TagSetIndices.Append(PkgList.PackageIndicies);
-					}
-				}
+				check(Container.ChunkIds[Idx].IsValid());
+				SortedChunkIds.Add(Container.ChunkIds[Idx]);
+				SortedChunkEntries.Add(Container.ChunkEntries[Idx]);
 			}
 
 			FOnDemandContainerEntry& ContainerEntry = ContainerEntries.AddDefaulted_GetRef();
-			ContainerEntry.ContainerId			= LegacyContainer.ContainerId;
-			ContainerEntry.ContainerName		= WriteString(LegacyContainer.ContainerName);
-			ContainerEntry.ContainerHeaderSize	= LegacyContainer.Header.Num();
-			ContainerEntry.ChunkCount			= LegacyContainer.Entries.Num();
-			ContainerEntry.BlockCount			= BlockSizes.Num();
-			ContainerEntry.BlockSize			= Toc.Header.BlockSize;
-			ContainerEntry.TagSetCount			= TagSets.Num();
-			ContainerEntry.TagSetIndicesCount	= TagSetIndices.Num();
-			ContainerEntry.UTocHash				= LegacyContainer.UTocHash;
-			ContainerEntry.ContainerFlags		= Toc.Header.Flags;
-			ContainerEntry.FileContainerFlags 	= LegacyContainer.ContainerFlags;
-
-			if (LegacyContainer.EncryptionKeyGuid.IsEmpty() == false)
-			{
-				ensure(FGuid::Parse(LegacyContainer.EncryptionKeyGuid, ContainerEntry.EncryptionKeyGuid));
-			}
+			ContainerEntry.ContainerId			= Container.Id;
+			ContainerEntry.ContainerName		= WriteString(Container.Name);
+			ContainerEntry.ContainerHeaderSize	= uint32(Container.Header.GetSize());
+			ContainerEntry.EncryptionKeyGuid	= Container.EncryptionKeyGuid;
+			ContainerEntry.ChunkCount			= Container.ChunkIds.Num();
+			ContainerEntry.BlockCount			= Container.BlockSizes.Num();
+			ContainerEntry.BlockSize			= Container.CompressionBlockSize; 
+			ContainerEntry.TagSetCount			= Container.TagSets.Num();
+			ContainerEntry.TagSetIndicesCount	= Container.TagSetIndices.Num();
+			ContainerEntry.UTocHash				= Container.UTocHash;
+			ContainerEntry.ContainerFlags		= uint32(Container.ContainerFlags);
+			ContainerEntry.FileContainerFlags 	= Container.FileContainerFlags;
+			ContainerEntry.PartitionCount		= uint32(Container.Partitions.Num());
 
 			ContainerEntry.DataOffset = IntCastChecked<uint32>(ContainerDataAr.Tell());
 			{
-				ContainerDataAr.Serialize(ChunkIds.GetData(), ChunkIds.GetTypeSize() * ChunkIds.Num());
-				ContainerDataAr.Serialize(ChunkEntries.GetData(), ChunkEntries.GetTypeSize() * ChunkEntries.Num());
-				ContainerDataAr.Serialize((void*)BlockSizes.GetData(), BlockSizes.GetTypeSize() * BlockSizes.Num());
-				ContainerDataAr.Serialize((void*)BlockHashes.GetData(), BlockHashes.GetTypeSize() * BlockHashes.Num());
-				ContainerDataAr.Serialize(TagSets.GetData(), TagSets.GetTypeSize() * TagSets.Num());
-				ContainerDataAr.Serialize(TagSetIndices.GetData(), TagSetIndices.GetTypeSize() * TagSetIndices.Num());
-				ContainerDataAr.Serialize((void*)LegacyContainer.Header.GetData(), LegacyContainer.Header.Num());
+				SerializeArrayData(ContainerDataAr, Container.Partitions);
+				SerializeArrayData(ContainerDataAr, SortedChunkIds);
+				SerializeArrayData(ContainerDataAr, SortedChunkEntries);
+				SerializeArrayData(ContainerDataAr, BlockSizes);
+				SerializeArrayData(ContainerDataAr, BlockHashes);
+				SerializeArrayData(ContainerDataAr, Container.TagSets);
+				SerializeArrayData(ContainerDataAr, Container.TagSetIndices);
+				ContainerDataAr.Serialize((void*)Container.Header.GetData(), Container.Header.GetSize());
 			}
 			ContainerEntry.DataSize = IntCastChecked<uint32>(ContainerDataAr.Tell()) - ContainerEntry.DataOffset;
 
-			ContainerIndex++;
+			if (Container.CompressionFormat != NAME_None)
+			{
+				check(CompressionFormat == NAME_None || CompressionFormat == Container.CompressionFormat);
+				if (CompressionFormat == NAME_None)
+				{
+					CompressionFormat = Container.CompressionFormat;
+				}
+			}
 		}
 
-		V1::FOnDemandTocHeader Header;
+		V2::FOnDemandTocHeader Header;
 		Header.EpochTimestamp		= FDateTime::Now().ToUnixTimestamp();
-		Header.BuildVersion			= WriteString(Toc.Meta.BuildVersion);
-		Header.TargetPlatform		= WriteString(Toc.Meta.TargetPlatform);
-		Header.ChunksDirectory		= WriteString(Toc.Header.ChunksDirectory);
-		Header.HostGroupName		= WriteString(Toc.Header.HostGroupName);
-		Header.CompressionFormat	= WriteString(Toc.Header.CompressionFormat);
-		Header.ContainerCount		= Toc.Containers.Num();
+		Header.BuildVersion			= WriteString(Metadata.BuildVersion);
+		Header.TargetPlatform		= WriteString(Metadata.TargetPlatform);
+		Header.ChunksDirectory		= WriteString(Metadata.ChunksDirectory);
+		Header.HostGroupName		= WriteString(Metadata.HostGroupName);
+		Header.CompressionFormat	= WriteString(CompressionFormat.ToString());
+		Header.ContainerCount		= Containers.Num();
 		Header.StringTableLen		= StringTableLen;
 
 		Ar.Serialize(&Header, sizeof(FOnDemandTocHeader));
@@ -288,9 +430,12 @@ private:
 
 	using FLookup = TMap<FUtf8StringView, FOnDemandStringEntry>;
 
-	FLookup						StringToEntry;
-	TChunkedArray<FUtf8String>	StringTable;
-	int32						StringTableLen = 0;
+	FLookup								StringToEntry;
+	TChunkedArray<FUtf8String>			StringTable;
+	TChunkedArray<FContainer>			Containers;
+	FMetadata							Metadata;
+	FContainer*							CurrentContainer = nullptr;
+	int32								StringTableLen = 0;
 };
 
 ////////////////////////////////////////////////////////////////////////////////
@@ -314,12 +459,95 @@ FOnDemandTocWriter& FOnDemandTocWriter::operator=(FOnDemandTocWriter&& Other)
 	return *this;
 }
 
-TIoStatusOr<uint64> FOnDemandTocWriter::Write(const UE::IoStore::FOnDemandToc& Toc, FArchive& Ar)
+bool FOnDemandTocWriter::IsEmpty() const
+{
+	return !Impl.IsValid() || Impl->IsEmpty();
+}
+
+void FOnDemandTocWriter::SetMetadata(const FString& BuildVersion, const FString& TargetPlatform)
+{
+	check(Impl.IsValid());
+	Impl->SetMetadata(BuildVersion, TargetPlatform);
+}
+
+void FOnDemandTocWriter::SetHostGroup(const FString& HostGroupName)
+{
+	check(Impl.IsValid());
+	Impl->SetHostGroup(HostGroupName);
+}
+
+void FOnDemandTocWriter::SetChunksDirectory(const FString& ChunksDirectory)
+{
+	check(Impl.IsValid());
+	Impl->SetChunksDirectory(ChunksDirectory);
+}
+
+void FOnDemandTocWriter::BeginContainer(
+	const FString& ContainerName,
+	const FIoContainerId& ContainerId,
+	const FIoHash& UTocHash,
+	const FGuid& EncryptionKeyGuid,
+	uint32 CompressionBlockSize,
+	uint32 FileContainerFlags,
+	EOnDemandContainerEntryFlags ContainerFlags,
+	FIoBuffer ContainerHeader)
 {
-	return Impl->Write(Toc, Ar);
+	check(Impl.IsValid());
+	Impl->BeginContainer(
+		ContainerName,
+		ContainerId,
+		UTocHash,
+		EncryptionKeyGuid,
+		CompressionBlockSize,
+		FileContainerFlags,
+		ContainerFlags,
+		ContainerHeader);
 }
 
-TIoStatusOr<uint64> FOnDemandTocWriter::Write(const UE::IoStore::FOnDemandToc& Toc, const FString& Filename)
+int32 FOnDemandTocWriter::BeginPartition()
+{
+	check(Impl.IsValid());
+	return Impl->BeginPartition();
+}
+
+uint64 FOnDemandTocWriter::AddChunk(
+	int32 Partition,
+	const FIoChunkId& ChunkId,
+	FIoBuffer EncodedChunk,
+	TConstArrayView<uint32> BlockSizes,
+	TConstArrayView<FIoBlockHash> BlockHashes,
+	uint64 RawChunkSize)
+{
+	check(Partition >= 0);
+	check(Impl.IsValid());
+	return Impl->AddChunk(Partition, ChunkId, EncodedChunk, BlockSizes, BlockHashes, RawChunkSize);
+}
+
+void FOnDemandTocWriter::AddTagSet(FStringView Tag, TConstArrayView<uint32> PackageIndices)
+{
+	check(Impl.IsValid());
+	Impl->AddTagSet(Tag, PackageIndices);
+}
+
+FIoBuffer FOnDemandTocWriter::EndPartition(int32 Partition, FIoHash& OutHash)
+{
+	check(Partition >= 0);
+	check(Impl.IsValid());
+	return Impl->EndPartition(Partition, OutHash);
+}
+
+void FOnDemandTocWriter::EndContainer(FName CompressionFormat)
+{
+	check(Impl.IsValid());
+	Impl->EndContainer(CompressionFormat);
+}
+
+TIoStatusOr<uint64> FOnDemandTocWriter::Write(FArchive& Ar)
+{
+	return Impl->Write(Ar);
+}
+
+TIoStatusOr<uint64> FOnDemandTocWriter::Write(const FString& Filename)
 {
 	TUniquePtr<FArchive> Ar(IFileManager::Get().CreateFileWriter(*Filename));
 	if (Ar.IsValid() == false)
@@ -328,7 +556,7 @@ TIoStatusOr<uint64> FOnDemandTocWriter::Write(const UE::IoStore::FOnDemandToc& T
 			<< TEXT("Failed to open '") << Filename << TEXT("'");
 	}
 
-	return Impl->Write(Toc, *Ar);
+	return Impl->Write(*Ar);
 }
 
 ////////////////////////////////////////////////////////////////////////////////
@@ -564,198 +792,281 @@ TIoStatusOr<FOnDemandTocReader> FOnDemandTocReader::Read(const FString& Filename
 
 } // V1 
 
-///////////////////////////////////////////////////////////////////////////////
-#if WITH_IOSTORE_ONDEMAND_TESTS
+////////////////////////////////////////////////////////////////////////////////
+namespace V2
+{
 
-class FTmpDirectoryScope
+FOnDemandTocReader::FOnDemandTocReader(const FString& InFilename, FIoBuffer InBuffer)
+	: Filename(InFilename)
+	, Buffer(InBuffer)
 {
-public:
-	explicit FTmpDirectoryScope(const FString& InDir)
-		: Ifm(IFileManager::Get())
-		, Dir(InDir)
+	Header = reinterpret_cast<const FOnDemandTocHeader*>(Buffer.GetData());
+	StringTable = reinterpret_cast<const UTF8CHAR*>(Header + 1);
+	ContainerEntries = MakeArrayView(
+		reinterpret_cast<const FOnDemandContainerEntry*>(StringTable + Header->StringTableLen), Header->ContainerCount);
+}
+
+FUtf8StringView FOnDemandTocReader::GetString(FOnDemandStringEntry StringEntry) const
+{
+	return FUtf8StringView(StringTable + StringEntry.Offset, StringEntry.Len);
+}
+
+TIoStatusOr<FOnDemandContainerTocView> FOnDemandTocReader::ReadContainer(
+	const FOnDemandContainerEntry& Container,
+	FOnDemandTocStorage& Storage,
+	EOnDemandTocReaderOptions Options) const
+{
+	const FOnDemandContainerEntry* Existing = Algo::FindByPredicate(
+		ContainerEntries, [&Container](const FOnDemandContainerEntry& Entry) { return &Entry == &Container; });
+
+	if (Existing == nullptr)
 	{
-		const bool bTree			= true;
-		const bool bRequireExists	= false;
-		Ifm.DeleteDirectory(*Dir, bRequireExists, bTree);
-		Ifm.MakeDirectory(*Dir, bTree);
+		return FIoStatus(EIoErrorCode::NotFound);
 	}
 
-	~FTmpDirectoryScope()
+	IPlatformFile& Ipf = FPlatformFileManager::Get().GetPlatformFile();
+
+	const int64 ContainerDataOffset = sizeof(FOnDemandTocHeader)
+		+ (Header->StringTableLen * sizeof(UTF8CHAR))
+		+ (Header->ContainerCount * sizeof(FOnDemandContainerEntry))
+		+ Container.DataOffset;
+
+	const int64 ContainerDataSize = int64(Container.DataSize) - int64(Container.ContainerHeaderSize);
+	if (ContainerDataSize <= 0)
 	{
-		const bool bTree			= true;
-		const bool bRequireExists	= false;
-		Ifm.DeleteDirectory(*Dir, bRequireExists, bTree);
+		return FIoStatusBuilder(EIoErrorCode::CorruptToc)
+			<< TEXT("Unexpected container data size in '")
+			<< Filename << TEXT("'");
+	}
+
+	const bool bMemoryMap = FPlatformProperties::SupportsMemoryMappedFiles()
+		&& EnumHasAnyFlags(Options, EOnDemandTocReaderOptions::MemoryMap);
+
+	if (bMemoryMap)
+	{
+		FOpenMappedResult OpenResult = Ipf.OpenMappedEx(*Filename);
+		if (OpenResult.HasError() == false)
+		{
+			TUniquePtr<IMappedFileHandle> MappedFile = OpenResult.StealValue();
+			TUniquePtr<IMappedFileRegion> Region(MappedFile->MapRegion(ContainerDataOffset, ContainerDataSize));
+
+			if (Region.IsValid())
+			{
+				Storage = FOnDemandTocStorage(Buffer, MoveTemp(MappedFile), MoveTemp(Region));
+			}
+		}
+	}
+
+	FFileOpenResult OpenResult = Ipf.OpenRead(*Filename, IPlatformFile::EOpenReadFlags::None);
+	if (OpenResult.HasError())
+	{
+		FFileSystemError& Error = OpenResult.GetError();
+		return FIoStatusBuilder(EIoErrorCode::FileOpenFailed)
+			<< TEXT("Failed to open '") << Filename << TEXT("' reason: ")
+			<< Error.GetMessage();
 	}
-private:
-	IFileManager& Ifm;
-	FString Dir;
-};
 
-TEST_CASE("IoStore::OnDemand::Serialization::TOC::V1", "[IoStoreOnDemand][Serialization]")
+	TUniquePtr<IFileHandle> FileHandle = OpenResult.StealValue();
+	if (Storage.IsEmpty())
+	{
+		if (FileHandle->Seek(ContainerDataOffset) == false)
+		{
+			return FIoStatusBuilder(EIoErrorCode::FileSeekFailed)
+				<< TEXT("Failed to seek to container data in '") << Filename << TEXT("'");
+		}
+
+		FIoBuffer ContainerData(ContainerDataSize);
+		if (FileHandle->Read(ContainerData.GetData(), ContainerData.GetSize()) == false)
+		{
+			return FIoStatusBuilder(EIoErrorCode::ReadError)
+				<< TEXT("Failed to read container data in '") << Filename << TEXT("'");
+		}
+
+		Storage = FOnDemandTocStorage(Buffer, ContainerData);
+	}
+
+	const int64 ContainerHeaderOffset = ContainerDataOffset + ContainerDataSize;
+	if (FileHandle->Seek(ContainerHeaderOffset) == false)
+	{
+		return FIoStatusBuilder(EIoErrorCode::FileSeekFailed)
+			<< TEXT("Failed to seek to container header in '") << Filename << TEXT("'");
+	}
+
+	// The container header is currently copied into a separate buffer
+	FIoBuffer ContainerHeaderChunk(Container.ContainerHeaderSize);
+	if (FileHandle->Read(ContainerHeaderChunk.GetData(), ContainerHeaderChunk.GetSize()) == false)
+	{
+		return FIoStatusBuilder(EIoErrorCode::ReadError)
+			<< TEXT("Failed to read container header in '") << Filename << TEXT("'");
+	}
+
+	FMemoryView ContainerData = Storage.GetView();
+	FOnDemandContainerTocView View;
+	View.Header = FOnDemandContainerTocHeaderView(Header, StringTable, &Container);
+
+	View.PartitionEntries = MakeArrayView(reinterpret_cast<const FOnDemandPartitionEntry*>(ContainerData.GetData()), Container.PartitionCount);
+	ContainerData += View.PartitionEntries.NumBytes();
+
+	View.ChunkIds = MakeArrayView(reinterpret_cast<const FIoChunkId*>(ContainerData.GetData()), Container.ChunkCount);
+	ContainerData += View.ChunkIds.NumBytes();
+
+	View.ChunkEntries = MakeArrayView(reinterpret_cast<const FOnDemandChunkEntry*>(ContainerData.GetData()), Container.ChunkCount);
+	ContainerData += View.ChunkEntries.NumBytes();
+
+	View.BlockSizes = MakeArrayView(reinterpret_cast<const uint32*>(ContainerData.GetData()), Container.BlockCount);
+	ContainerData += View.BlockSizes.NumBytes();
+
+	View.BlockHashes = MakeArrayView(reinterpret_cast<const FIoBlockHash*>(ContainerData.GetData()), Container.BlockCount);
+	ContainerData += View.BlockHashes.NumBytes();
+
+	View.TagSets = MakeArrayView(reinterpret_cast<const FOnDemandTagSetEntry*>(ContainerData.GetData()), Container.TagSetCount);
+	ContainerData += View.TagSets.NumBytes();
+
+	View.TagSetIndices = MakeArrayView(reinterpret_cast<const uint32*>(ContainerData.GetData()), Container.TagSetIndicesCount);
+	ContainerData += View.TagSetIndices.NumBytes();
+
+	if (ContainerData.IsEmpty() == false)
+	{
+		return FIoStatusBuilder(EIoErrorCode::CorruptToc)
+			<< TEXT("Unexpected container data size in '") << Filename << TEXT("'");
+	}
+
+	View.ContainerHeaderChunk = MoveTemp(ContainerHeaderChunk);
+
+	return View;
+}
+
+TIoStatusOr<FOnDemandTocReader> FOnDemandTocReader::Read(const FString& Filename)
 {
-	using namespace UE::IoStore::Serialization;
-	using namespace UE::IoStore::Serialization::V1;
+	IPlatformFile& Ipf = FPlatformFileManager::Get().GetPlatformFile();
 
-	const FString TestBaseDir = "TocTestTmpDir";
+	FFileOpenResult OpenResult = Ipf.OpenRead(*Filename, IPlatformFile::EOpenReadFlags::None);
+	if (OpenResult.HasError())
+	{
+		FFileSystemError& Error = OpenResult.GetError();
+		return FIoStatusBuilder(EIoErrorCode::FileOpenFailed)
+			<< TEXT("Failed to open '") << Filename << TEXT("' reason: ")
+			<< Error.GetMessage();
+	}
+
+	TUniquePtr<IFileHandle> FileHandle = OpenResult.StealValue();
 
-	SECTION("SetBlockInfoFromOffsetAndCount")
+	FOnDemandTocHeader Header;
+	if (FileHandle->Read(reinterpret_cast<uint8*>(&Header), sizeof(FOnDemandTocHeader)) == false)
 	{
-		const FOnDemandChunkBlockInfo Nfo = FOnDemandChunkBlockInfo::FromOffsetAndCount(46, 2);
-		CHECK(Nfo.Size() == 0);
-		CHECK(Nfo.Hash() == 0);
-		CHECK(Nfo.Offset() == 46);
-		CHECK(Nfo.Count() == 2);
+		return FIoStatusBuilder(EIoErrorCode::ReadError)
+			<< TEXT("Failed to read container TOC header from '") << Filename << TEXT("'");
 	}
 
-	SECTION("SetBlockInfoFromSizeAndHash")
+	if (Header.Signature.IsValid() == false)
 	{
-		const FOnDemandChunkBlockInfo Nfo = FOnDemandChunkBlockInfo::FromSizeAndHash(46, 2);
-		CHECK(Nfo.Size() == 46);
-		CHECK(Nfo.Hash() == 2);
-		CHECK(Nfo.Offset() == ~uint32(0));
-		CHECK(Nfo.Count() == 1);
+		return FIoStatusBuilder(EIoErrorCode::CorruptToc)
+			<< TEXT("Invalid TOC signature found in '") << Filename << TEXT("'");
 	}
 
-	SECTION("ConvertFromLegacyToc")
+	if (Header.Version.IsValid() == false)
 	{
-		const FString TocFilePath = TestBaseDir / TEXT("test.uondemandtoc");
-		FTmpDirectoryScope _(TestBaseDir);
-		UE::IoStore::FOnDemandToc LegacyToc;
+		return FIoStatusBuilder(EIoErrorCode::CorruptToc)
+			<< TEXT("Invalid TOC version found in '") << Filename << TEXT("'");
+	}
 
-		// Create a legacy TOC
-		{
-			auto AddEntry = [](FOnDemandTocContainerEntry& Container, FOnDemandTocEntry Entry)
-			{
-				check(Entry.BlockCount > 0);
-				Entry.BlockOffset = Container.BlockSizes.Num();
-				for (uint32 Idx = 0; Idx < Entry.BlockCount; ++Idx)
-				{
-					Container.BlockSizes.Add(Idx + 1);
-					Container.BlockHashes.Add(Idx + 1);
-				}
-				Container.Entries.Add(Entry);
-			};
+	if (Header.Version.Major != uint16(EOnDemandTocMajorVersion::Latest) || Header.Version.Minor != uint16(EOnDemandTocMinorVersion::Latest))
+	{
+		return FIoStatusBuilder(EIoErrorCode::InvalidParameter)
+			<< TEXT("Unsupported TOC version found in '") << Filename << TEXT("'");
+	}
 
-			LegacyToc.Header.BlockSize = 64 << 10;
-			LegacyToc.Header.CompressionFormat = TEXT("Oodle");
-			LegacyToc.Header.ChunksDirectory = TEXT("iostore/test");
-			LegacyToc.Header.HostGroupName = TEXT("myhost");
+	const int64 TotalHeaderSize =
+		int64(sizeof(FOnDemandTocHeader)) +
+		int64((sizeof(UTF8CHAR) * Header.StringTableLen)) +
+		int64((sizeof(FOnDemandContainerEntry) * Header.ContainerCount));
 
-			FOnDemandTocContainerEntry& Container = LegacyToc.Containers.AddDefaulted_GetRef();
-			Container.ContainerName = TEXT("test");
-			Container.ContainerId = FIoContainerId::FromName(FName(*Container.ContainerName));
+	if ((TotalHeaderSize < int64(sizeof(FOnDemandTocHeader))) || (TotalHeaderSize >= FileHandle->Size()))
+	{
+		return FIoStatusBuilder(EIoErrorCode::CorruptToc)
+			<< TEXT("TOC header size greater than file size in '") << Filename << TEXT("'");
+	}
 
-			AddEntry(Container, FOnDemandTocEntry
-			{
-				.ChunkId	= CreatePackageDataChunkId(FPackageId::FromName(FName(TEXT("MyAssetA")))),
-				.RawSize	= 42,
-				.BlockCount = 1
-			});
-			AddEntry(Container, FOnDemandTocEntry
-			{
-				.ChunkId	= CreatePackageDataChunkId(FPackageId::FromName(FName(TEXT("MyAssetB")))),
-				.RawSize	= 42,
-				.BlockCount = 2
-			});
-			AddEntry(Container, FOnDemandTocEntry
-			{
-				.ChunkId	= CreatePackageDataChunkId(FPackageId::FromName(FName(TEXT("MyAssetC")))),
-				.RawSize	= 42,
-				.BlockCount = 3
-			});
+	FIoBuffer Buffer(TotalHeaderSize);
+	FMutableMemoryView BufferView = Buffer.GetMutableView();
+	FMemory::Memcpy(BufferView.GetData(), &Header, sizeof(FOnDemandTocHeader));
+	BufferView += sizeof(FOnDemandTocHeader);
 
-			FOnDemandTocTagSet& TagSet = LegacyToc.TagSets.AddDefaulted_GetRef();
-			TagSet.Tag = TEXT("TheTag");
-			FOnDemandTocTagSetPackageList& List = TagSet.Packages.AddDefaulted_GetRef();
-			List.ContainerIndex = 0;
-			List.PackageIndicies.Add(1);
-		}
+	if (FileHandle->Read(reinterpret_cast<uint8*>(BufferView.GetData()), BufferView.GetSize()) == false)
+	{
+		return FIoStatusBuilder(EIoErrorCode::ReadError)
+			<< TEXT("Failed to read container TOC header from '") << Filename << TEXT("'");
+	}
 
-		// Serialize the legacy TOC to the latest format
-		{
+	if (FileHandle->Seek(FileHandle->Size() - int64(sizeof(FOnDemandTocFooter))) == false)
+	{
+		return FIoStatusBuilder(EIoErrorCode::FileSeekFailed)
+			<< TEXT("Failed to seek to container TOC footer in '") << Filename << TEXT("'");
+	}
 
-			FOnDemandTocWriter Writer;
-			TIoStatusOr<uint64> Status = Writer.Write(LegacyToc, TocFilePath);
-			CHECK(Status.IsOk());
-		}
+	FOnDemandTocFooter Footer;
+	if (FileHandle->Read(reinterpret_cast<uint8*>(&Footer), sizeof(FOnDemandTocFooter)) == false)
+	{
+		return FIoStatusBuilder(EIoErrorCode::ReadError)
+			<< TEXT("Failed to read container TOC footer from '") << Filename << TEXT("'");
+	}
 
-		// Assert
-		FOnDemandTocReader Reader = FOnDemandTocReader::Read(TocFilePath).ConsumeValueOrDie();
-		CHECK(FString(Reader.ChunksDirectory()) == LegacyToc.Header.ChunksDirectory);
-		CHECK(FString(Reader.CompressionFormat()) == LegacyToc.Header.CompressionFormat);
-		CHECK(FString(Reader.HostGroupName()) == LegacyToc.Header.HostGroupName);
+	if (Footer.Signature.IsValid() == false)
+	{
+		return FIoStatusBuilder(EIoErrorCode::CorruptToc)
+			<< TEXT("Invalid TOC footer signature found in '") << Filename << TEXT("'");
+	}
 
-		for (int32 ContainerIndex = 0; const FOnDemandTocContainerEntry& LegacyContainer : LegacyToc.Containers)
-		{
-			TConstArrayView<uint32>			LegacyContainerBlockSizes = LegacyContainer.BlockSizes;
-			TConstArrayView<FIoBlockHash>	LegacyContainerBlockHashes = LegacyContainer.BlockHashes;
+	return FOnDemandTocReader(Filename, Buffer);
+}
 
-			const FOnDemandContainerEntry& ContainerEntry = Reader.Containers()[ContainerIndex];
+} // V2
 
-			FOnDemandTocStorage Storage;
-			FOnDemandContainerTocView Container =
-				Reader.ReadContainer(ContainerEntry, Storage, EOnDemandTocReaderOptions::None)
-					.ConsumeValueOrDie();
+} // UE::IoStore::Serialization
+
+////////////////////////////////////////////////////////////////////////////////
+FStringBuilderBase& operator<<(FStringBuilderBase& Sb, UE::IoStore::Serialization::EOnDemandContainerEntryFlags Flags)
+{
+	using namespace UE::IoStore::Serialization;
 
-			// Check header
-			CHECK(FString(Container.Header.ContainerName()) == LegacyContainer.ContainerName);
-			CHECK(Container.Header.ContainerId() == LegacyContainer.ContainerId);
-			CHECK(Container.Header.BlockSize() == LegacyToc.Header.BlockSize);
+	static const TCHAR* FlagText[]
+	{
+		TEXT("None"),
+		TEXT("InstallOnDemand"),
+		TEXT("StreamOnDemand"),
+	};
 
-			// Check entries
-			for (const FOnDemandTocEntry& LegacyEntry : LegacyContainer.Entries)
-			{
-				const int32 EntryIndex = Algo::LowerBound(Container.ChunkIds, LegacyEntry.ChunkId);
-				CHECK(EntryIndex != INDEX_NONE);
-				CHECK(EntryIndex < Container.ChunkIds.Num());
-
-				const FOnDemandChunkEntry& Entry = Container.ChunkEntries[EntryIndex];
-				CHECK(Entry.RawSize == LegacyEntry.RawSize);
-				CHECK(Entry.BlockInfo.Count() == LegacyEntry.BlockCount);
-
-				// Check compression block size and hash
-				TConstArrayView<uint32> LegacyBlockSizes = LegacyContainerBlockSizes.Mid(LegacyEntry.BlockOffset, LegacyEntry.BlockCount);
-				TConstArrayView<uint32> BlockSizes = Entry.BlockInfo.Count() > 1
-					? Container.BlockSizes.Mid(Entry.BlockInfo.Offset(), Entry.BlockInfo.Count())
-					: TConstArrayView<uint32>(Entry.BlockInfo.Data(), 1);
-				
-				for (int32 Idx = 0; Idx < BlockSizes.Num(); ++Idx)
-				{
-					CHECK(BlockSizes[Idx] == LegacyBlockSizes[Idx]);
-				}
+	if (Flags == EOnDemandContainerEntryFlags::None)
+	{
+		Sb << FlagText[0];
+		return Sb;
+	}
 
-				TConstArrayView<uint32> LegacyBlockHashes = LegacyContainerBlockHashes.Mid(LegacyEntry.BlockOffset, LegacyEntry.BlockCount);
-				TConstArrayView<uint32> BlockHashes = Entry.BlockInfo.Count() > 1
-					? Container.BlockHashes.Mid(Entry.BlockInfo.Offset(), Entry.BlockInfo.Count())
-					: TConstArrayView<uint32>(Entry.BlockInfo.Data() + 1, 1);
-				
-				for (int32 Idx = 0; Idx < BlockSizes.Num(); ++Idx)
-				{
-					CHECK(BlockHashes[Idx] == LegacyBlockHashes[Idx]);
-				}
+	constexpr uint32 BitCount = 1 + FMath::CountTrailingZeros(
+		static_cast<std::underlying_type_t<EOnDemandContainerEntryFlags>>(EOnDemandContainerEntryFlags::Last));
+	static_assert(UE_ARRAY_COUNT(FlagText) == BitCount + 1, "Please update flag text list");
 
-				// Check tags
-				for (const FOnDemandTagSetEntry& TagEntry : Container.TagSets)
-				{
-					FString TagName = FString(Container.Header.GetString(TagEntry.Tag));
-					const FOnDemandTocTagSet* LegacyTagSet = Algo::FindBy(LegacyToc.TagSets, TagName, &FOnDemandTocTagSet::Tag);
-					CHECK(LegacyTagSet != nullptr);
-					const FOnDemandTocTagSetPackageList* List = Algo::FindBy(LegacyTagSet->Packages, ContainerIndex, &FOnDemandTocTagSetPackageList::ContainerIndex);
-					CHECK(List != nullptr);
-
-					TConstArrayView<uint32> Indices = Container.TagSetIndices.Mid(TagEntry.Offset, TagEntry.Count);
-					for (int32 Idx = 0; Idx < Indices.Num(); ++Idx)
-					{
-						CHECK(Indices[Idx] == List->PackageIndicies[Idx]);
-					}
-				}
+	for (int32 Idx = 0; Idx < BitCount; ++Idx)
+	{
+		const EOnDemandContainerEntryFlags FlagToTest = static_cast<EOnDemandContainerEntryFlags>(1 << Idx);
+		if (EnumHasAnyFlags(Flags, FlagToTest))
+		{
+			if (Sb.Len())
+			{
+				Sb << TEXT("|");
 			}
-
-			ContainerIndex++;
+			Sb << FlagText[Idx + 1];
 		}
 	}
+
+	return Sb;
 }
 
-#endif // WITH_IOSTORE_ONDEMAND_TESTS
+////////////////////////////////////////////////////////////////////////////////
+FString LexToString(UE::IoStore::Serialization::EOnDemandContainerEntryFlags Flags)
+{
+	TStringBuilder<128> Sb;
+	Sb << Flags;
+	return FString::ConstructFromPtrSize(Sb.ToString(), Sb.Len());
+}
 
-} // UE::IoStore::Serialization
diff --git a/Engine/Source/Runtime/Experimental/IoStore/OnDemandCore/Public/IO/HttpIoDispatcher.h b/Engine/Source/Runtime/Experimental/IoStore/OnDemandCore/Public/IO/HttpIoDispatcher.h
--- a/Engine/Source/Runtime/Experimental/IoStore/OnDemandCore/Public/IO/HttpIoDispatcher.h
+++ b/Engine/Source/Runtime/Experimental/IoStore/OnDemandCore/Public/IO/HttpIoDispatcher.h
@@ -67,10 +67,10 @@ struct FIoHttpRange
 	[[nodiscard]] inline FIoOffsetAndLength ToOffsetAndLength() const;
 	static inline FIoHttpRange			FromOffsetAndLength(const FIoOffsetAndLength& OffsetAndLength);
 
-	[[nodiscard]] bool					operator==(const FIoHttpRange& Other)	{ return Min == Other.Min && Max == Other.Max; }
-	[[nodiscard]] bool					operator!=(const FIoHttpRange& Other)	{ return Min != Other.Min || Max != Other.Max; }
-	[[nodiscard]] FIoHttpRange			operator+(const FIoHttpRange& Other)	{ return Expand(Other); }
-	FIoHttpRange&						operator+=(const FIoHttpRange& Other)	{ return ExpandInline(Other); }
+	[[nodiscard]] bool					operator==(const FIoHttpRange& Other)	const	{ return Min == Other.Min && Max == Other.Max; }
+	[[nodiscard]] bool					operator!=(const FIoHttpRange& Other)	const	{ return Min != Other.Min || Max != Other.Max; }
+	[[nodiscard]] FIoHttpRange			operator+(const FIoHttpRange& Other)	const	{ return Expand(Other); }
+	FIoHttpRange&						operator+=(const FIoHttpRange& Other)			{ return ExpandInline(Other); }
 
 private:
 	uint32 Min = MAX_uint32;
diff --git a/Engine/Source/Runtime/Experimental/IoStore/OnDemandCore/Public/IO/OnDemandError.h b/Engine/Source/Runtime/Experimental/IoStore/OnDemandCore/Public/IO/OnDemandError.h
--- a/Engine/Source/Runtime/Experimental/IoStore/OnDemandCore/Public/IO/OnDemandError.h
+++ b/Engine/Source/Runtime/Experimental/IoStore/OnDemandCore/Public/IO/OnDemandError.h
@@ -194,8 +194,8 @@ struct FChunkMissingErrorContext
 struct FChunkHashMismatchErrorContext
 {
 	FIoChunkId	ChunkId;
-	FIoHash		ExpectedHash;
-	FIoHash		ActualHash;
+	FString		ExpectedHash;
+	FString		ActualHash;
 	UE::IoStore::EOnDemandInstallCasType CasType = UE::IoStore::EOnDemandInstallCasType::None;
 };
 
@@ -243,7 +243,6 @@ inline void SerializeForLog(FCbWriter& Writer, const UE::UnifiedError::IoStoreOn
 	Writer.EndObject();
 }
 
-
 UE_API void SerializeForLog(FCbWriter& Writer, const FIoChunkId& ChunkId);
 
 ////////////////////////////////////////////////////////////////////////////////
@@ -258,8 +257,8 @@ inline void SerializeForLog(FCbWriter& Writer, const UE::UnifiedError::IoStoreOn
 	Writer.AddString(UTF8TEXTVIEW("$type"), TErrorStructFeatures<UE::UnifiedError::IoStoreOnDemand::FChunkHashMismatchErrorContext>::GetErrorContextTypeNameAsString());
 	Writer.AddString(UTF8TEXTVIEW("$format"), TEXT("({CasType} CAS - ChunkId: {ChunkId}, ExpectedHash: {ExpectedHash}, ActualHash: {ActualHash})"));
 	Writer.AddString(UTF8TEXTVIEW("ChunkId"), LexToString(Ctx.ChunkId));
-	Writer.AddString(UTF8TEXTVIEW("ExpectedHash"), LexToString(Ctx.ExpectedHash));
-	Writer.AddString(UTF8TEXTVIEW("ActualHash"), LexToString(Ctx.ActualHash));
+	Writer.AddString(UTF8TEXTVIEW("ExpectedHash"), Ctx.ExpectedHash);
+	Writer.AddString(UTF8TEXTVIEW("ActualHash"), Ctx.ActualHash);
 	Writer.AddString(UTF8TEXTVIEW("CasType"), LexToString(Ctx.CasType));
 	Writer.EndObject();
 }
